Generic Instanced Actors engine-level plugin's stub

diff --git a/Engine/Plugins/Runtime/InstancedActors/InstancedActors.uplugin b/Engine/Plugins/Runtime/InstancedActors/InstancedActors.uplugin
new file mode 100644
--- /dev/null
+++ b/Engine/Plugins/Runtime/InstancedActors/InstancedActors.uplugin
@@ -0,0 +1,46 @@
+{
+	"FileVersion": 3,
+	"Version": 1,
+	"VersionName": "0.1",
+	"FriendlyName": "InstancedActors",
+	"Description": "",
+	"Category": "Gameplay",
+	"CreatedBy": "Epic Games, Inc.",
+	"CreatedByURL": "http://epicgames.com",
+	"DocsURL": "",
+	"MarketplaceURL": "",
+	"SupportURL": "",
+	"EnabledByDefault": false,
+	"CanContainContent": true,
+	"IsBetaVersion": false,
+	"IsExperimentalVersion": true,
+	"Installed": false,
+	"Modules": [
+		{
+			"Name": "InstancedActors",
+			"Type": "Runtime"
+		},
+		{
+			"Name": "InstancedActorsTestSuite",
+			"Type": "UncookedOnly"
+		}
+	],
+	"Plugins": [
+		{
+			"Name": "MassEntity",
+			"Enabled": true
+		},
+		{
+			"Name": "MassGameplay",
+			"Enabled": true
+		},
+		{
+			"Name": "MassAI",
+			"Enabled": true
+		},
+		{
+			"Name": "StructUtils",
+			"Enabled": true
+		},
+	]
+}
\ No newline at end of file
diff --git a/Engine/Plugins/Runtime/InstancedActors/Source/InstancedActors/InstancedActors.Build.cs b/Engine/Plugins/Runtime/InstancedActors/Source/InstancedActors/InstancedActors.Build.cs
new file mode 100644
--- /dev/null
+++ b/Engine/Plugins/Runtime/InstancedActors/Source/InstancedActors/InstancedActors.Build.cs
@@ -0,0 +1,45 @@
+// Copyright Epic Games, Inc. All Rights Reserved.
+
+namespace UnrealBuildTool.Rules
+{
+	public class InstancedActors : ModuleRules
+	{
+		public InstancedActors(ReadOnlyTargetRules Target) : base(Target)
+		{
+			PCHUsage = ModuleRules.PCHUsageMode.UseExplicitOrSharedPCHs;
+
+			PublicIncludePaths.AddRange(
+				new string[] {
+				"Runtime/AIModule/Public",
+				ModuleDirectory + "/Public",
+				}
+			);
+
+
+			PrivateIncludePaths.AddRange(
+				new string[] {
+				}
+			);
+
+			PublicDependencyModuleNames.AddRange(
+				new string[] {
+					"Core",
+					"CoreUObject",
+					"Engine",
+					"MassEntity",
+					"MassActors",
+					"MassRepresentation",
+					"MassSignals",
+					"MassSpawner",
+					"MassLOD",
+					"StructUtils",
+				}
+			);
+
+			if (Target.bBuildEditor == true)
+			{
+				PrivateDependencyModuleNames.Add("UnrealEd");
+			}
+		}
+	}
+}
diff --git a/Engine/Plugins/Runtime/InstancedActors/Source/InstancedActors/Private/InstancedActorsData.cpp b/Engine/Plugins/Runtime/InstancedActors/Source/InstancedActors/Private/InstancedActorsData.cpp
new file mode 100644
--- /dev/null
+++ b/Engine/Plugins/Runtime/InstancedActors/Source/InstancedActors/Private/InstancedActorsData.cpp
@@ -0,0 +1,3 @@
+// Copyright Epic Games, Inc. All Rights Reserved.
+
+#include "InstancedActorsData.h"
diff --git a/Engine/Plugins/Runtime/InstancedActors/Source/InstancedActors/Private/InstancedActorsManager.cpp b/Engine/Plugins/Runtime/InstancedActors/Source/InstancedActors/Private/InstancedActorsManager.cpp
new file mode 100644
--- /dev/null
+++ b/Engine/Plugins/Runtime/InstancedActors/Source/InstancedActors/Private/InstancedActorsManager.cpp
@@ -0,0 +1,3 @@
+// Copyright Epic Games, Inc. All Rights Reserved.
+
+#include "InstancedActorsManager.h"
diff --git a/Engine/Plugins/Runtime/InstancedActors/Source/InstancedActors/Private/InstancedActorsModule.cpp b/Engine/Plugins/Runtime/InstancedActors/Source/InstancedActors/Private/InstancedActorsModule.cpp
new file mode 100644
--- /dev/null
+++ b/Engine/Plugins/Runtime/InstancedActors/Source/InstancedActors/Private/InstancedActorsModule.cpp
@@ -0,0 +1,28 @@
+// Copyright Epic Games, Inc. All Rights Reserved.
+
+#include "IInstancedActorsModule.h"
+
+
+class FInstancedActorsModuleModule : public IInstancedActorsModule
+{
+	/** IModuleInterface implementation */
+	virtual void StartupModule() override;
+	virtual void ShutdownModule() override;
+};
+
+IMPLEMENT_MODULE(FInstancedActorsModuleModule, InstancedActors)
+
+
+
+void FInstancedActorsModuleModule::StartupModule()
+{
+	// This code will execute after your module is loaded into memory (but after global variables are initialized, of course.)
+}
+
+
+void FInstancedActorsModuleModule::ShutdownModule()
+{
+	// This function may be called during shutdown to clean up your module.  For modules that support dynamic reloading,
+	// we call this function before unloading the module.
+}
+
diff --git a/Engine/Plugins/Runtime/InstancedActors/Source/InstancedActors/Public/IInstancedActorsModule.h b/Engine/Plugins/Runtime/InstancedActors/Source/InstancedActors/Public/IInstancedActorsModule.h
new file mode 100644
--- /dev/null
+++ b/Engine/Plugins/Runtime/InstancedActors/Source/InstancedActors/Public/IInstancedActorsModule.h
@@ -0,0 +1,39 @@
+// Copyright Epic Games, Inc. All Rights Reserved.
+
+#pragma once
+
+#include "Modules/ModuleInterface.h"
+#include "Modules/ModuleManager.h"
+
+
+/**
+ * The public interface to this module.  In most cases, this interface is only public to sibling modules 
+ * within this plugin.
+ */
+class IInstancedActorsModule : public IModuleInterface
+{
+
+public:
+
+	/**
+	 * Singleton-like access to this module's interface.  This is just for convenience!
+	 * Beware of calling this during the shutdown phase, though.  Your module might have been unloaded already.
+	 *
+	 * @return Returns singleton instance, loading the module on demand if needed
+	 */
+	static inline IInstancedActorsModule& Get()
+	{
+		return FModuleManager::LoadModuleChecked<IInstancedActorsModule>( "InstancedActors" );
+	}
+
+	/**
+	 * Checks to see if this module is loaded and ready.  It is only valid to call Get() if IsAvailable() returns true.
+	 *
+	 * @return True if the module is loaded and ready to use
+	 */
+	static inline bool IsAvailable()
+	{
+		return FModuleManager::Get().IsModuleLoaded( "InstancedActors" );
+	}
+};
+
diff --git a/Engine/Plugins/Runtime/InstancedActors/Source/InstancedActors/Public/InstancedActorsData.h b/Engine/Plugins/Runtime/InstancedActors/Source/InstancedActors/Public/InstancedActorsData.h
new file mode 100644
--- /dev/null
+++ b/Engine/Plugins/Runtime/InstancedActors/Source/InstancedActors/Public/InstancedActorsData.h
@@ -0,0 +1,13 @@
+// Copyright Epic Games, Inc. All Rights Reserved.
+
+#pragma once
+
+#include "UObject/Object.h"
+#include "InstancedActorsData.generated.h"
+
+
+UCLASS(DefaultToInstanced)
+class INSTANCEDACTORS_API UInstancedActorsData : public UObject
+{
+	GENERATED_BODY()
+};
diff --git a/Engine/Plugins/Runtime/InstancedActors/Source/InstancedActors/Public/InstancedActorsManager.h b/Engine/Plugins/Runtime/InstancedActors/Source/InstancedActors/Public/InstancedActorsManager.h
new file mode 100644
--- /dev/null
+++ b/Engine/Plugins/Runtime/InstancedActors/Source/InstancedActors/Public/InstancedActorsManager.h
@@ -0,0 +1,13 @@
+// Copyright Epic Games, Inc. All Rights Reserved.
+
+#pragma once
+
+#include "ActorPartition/PartitionActor.h"
+#include "InstancedActorsManager.generated.h"
+
+
+UCLASS(Config=Game)
+class INSTANCEDACTORS_API AInstancedActorsManager : public APartitionActor
+{
+	GENERATED_BODY()
+};
diff --git a/Engine/Plugins/Runtime/InstancedActors/Source/InstancedActorsTestSuite/InstancedActorsTestSuite.Build.cs b/Engine/Plugins/Runtime/InstancedActors/Source/InstancedActorsTestSuite/InstancedActorsTestSuite.Build.cs
new file mode 100644
--- /dev/null
+++ b/Engine/Plugins/Runtime/InstancedActors/Source/InstancedActorsTestSuite/InstancedActorsTestSuite.Build.cs
@@ -0,0 +1,25 @@
+// Copyright Epic Games, Inc. All Rights Reserved.
+
+using UnrealBuildTool;
+
+namespace UnrealBuildTool.Rules
+{
+	public class InstancedActorsTestSuite : ModuleRules
+	{
+		public InstancedActorsTestSuite(ReadOnlyTargetRules Target) : base(Target)
+		{
+			PCHUsage = ModuleRules.PCHUsageMode.UseExplicitOrSharedPCHs;
+
+			PublicDependencyModuleNames.AddRange(
+				new string[] {
+					"Core",
+					"CoreUObject",
+					"Engine",
+					"MassEntity",
+					"InstancedActors",
+					"AITestSuite",
+				}
+			);
+		}
+	}
+}
\ No newline at end of file
diff --git a/Engine/Plugins/Runtime/InstancedActors/Source/InstancedActorsTestSuite/Private/InstancedActorsTest.cpp b/Engine/Plugins/Runtime/InstancedActors/Source/InstancedActorsTestSuite/Private/InstancedActorsTest.cpp
new file mode 100644
--- /dev/null
+++ b/Engine/Plugins/Runtime/InstancedActors/Source/InstancedActorsTestSuite/Private/InstancedActorsTest.cpp
@@ -0,0 +1,11 @@
+// Copyright Epic Games, Inc. All Rights Reserved.
+
+#include "AITestsCommon.h"
+
+#define LOCTEXT_NAMESPACE "InstancedActorsTest"
+
+UE_DISABLE_OPTIMIZATION_SHIP
+
+UE_ENABLE_OPTIMIZATION_SHIP
+
+#undef LOCTEXT_NAMESPACE
diff --git a/Engine/Plugins/Runtime/InstancedActors/Source/InstancedActorsTestSuite/Private/InstancedActorsTestSuiteModule.cpp b/Engine/Plugins/Runtime/InstancedActors/Source/InstancedActorsTestSuite/Private/InstancedActorsTestSuiteModule.cpp
new file mode 100644
--- /dev/null
+++ b/Engine/Plugins/Runtime/InstancedActors/Source/InstancedActorsTestSuite/Private/InstancedActorsTestSuiteModule.cpp
@@ -0,0 +1,13 @@
+// Copyright Epic Games, Inc. All Rights Reserved.
+
+#include "IInstancedActorsTestSuiteModule.h"
+
+#define LOCTEXT_NAMESPACE "InstancedActorsTestSuite"
+
+class FInstancedActorsTestSuiteModule : public IInstancedActorsTestSuiteModule
+{
+};
+
+IMPLEMENT_MODULE(FInstancedActorsTestSuiteModule, InstancedActorsTestSuite)
+
+#undef LOCTEXT_NAMESPACE
diff --git a/Engine/Plugins/Runtime/InstancedActors/Source/InstancedActorsTestSuite/Public/IInstancedActorsTestSuiteModule.h b/Engine/Plugins/Runtime/InstancedActors/Source/InstancedActorsTestSuite/Public/IInstancedActorsTestSuiteModule.h
new file mode 100644
--- /dev/null
+++ b/Engine/Plugins/Runtime/InstancedActors/Source/InstancedActorsTestSuite/Public/IInstancedActorsTestSuiteModule.h
@@ -0,0 +1,36 @@
+// Copyright Epic Games, Inc. All Rights Reserved.
+
+#pragma  once
+
+#include "Modules/ModuleManager.h"
+
+/**
+* The public interface to this module
+*/
+class IInstancedActorsTestSuiteModule : public IModuleInterface
+{
+
+public:
+
+	/**
+	* Singleton-like access to this module's interface.  This is just for convenience!
+	* Beware of calling this during the shutdown phase, though.  Your module might have been unloaded already.
+	*
+	* @return Returns singleton instance, loading the module on demand if needed
+	*/
+	static inline IInstancedActorsTestSuiteModule& Get()
+	{
+		return FModuleManager::LoadModuleChecked<IInstancedActorsTestSuiteModule>("InstancedActorsTestSuite");
+	}
+
+	/**
+	* Checks to see if this module is loaded and ready.  It is only valid to call Get() if IsAvailable() returns true.
+	*
+	* @return True if the module is loaded and ready to use
+	*/
+	static inline bool IsAvailable()
+	{
+		return FModuleManager::Get().IsModuleLoaded("InstancedActorsTestSuite");
+	}
+};
+
