CIS fix: added missing APartitionActor::GetDefaultGridSize override to InstancedActorManager.

diff --git a/Engine/Plugins/Runtime/InstancedActors/Source/InstancedActors/Private/InstancedActorsManager.cpp b/Engine/Plugins/Runtime/InstancedActors/Source/InstancedActors/Private/InstancedActorsManager.cpp
--- a/Engine/Plugins/Runtime/InstancedActors/Source/InstancedActors/Private/InstancedActorsManager.cpp
+++ b/Engine/Plugins/Runtime/InstancedActors/Source/InstancedActors/Private/InstancedActorsManager.cpp
@@ -1,3 +1,12 @@
 // Copyright Epic Games, Inc. All Rights Reserved.
 
 #include "InstancedActorsManager.h"
+
+
+#if WITH_EDITOR
+uint32 AInstancedActorsManager::GetDefaultGridSize(UWorld* InWorld) const
+{
+	ensureMsgf(false && "NotImplemented", TEXT("NotImplemented"));
+	return 0;
+}
+#endif // WITH_EDITOR
\ No newline at end of file
diff --git a/Engine/Plugins/Runtime/InstancedActors/Source/InstancedActors/Public/InstancedActorsManager.h b/Engine/Plugins/Runtime/InstancedActors/Source/InstancedActors/Public/InstancedActorsManager.h
--- a/Engine/Plugins/Runtime/InstancedActors/Source/InstancedActors/Public/InstancedActorsManager.h
+++ b/Engine/Plugins/Runtime/InstancedActors/Source/InstancedActors/Public/InstancedActorsManager.h
@@ -10,4 +10,9 @@ UCLASS(Config=Game)
 class INSTANCEDACTORS_API AInstancedActorsManager : public APartitionActor
 {
 	GENERATED_BODY()
+
+#if WITH_EDITOR
+	//~ Begin AActor Interface
+	virtual uint32 GetDefaultGridSize(UWorld* InWorld) const override;
+#endif
 };
