[InstancedActors] Removed unused plugins

diff --git a/Engine/Plugins/Runtime/InstancedActors/InstancedActors.uplugin b/Engine/Plugins/Runtime/InstancedActors/InstancedActors.uplugin
--- a/Engine/Plugins/Runtime/InstancedActors/InstancedActors.uplugin
+++ b/Engine/Plugins/Runtime/InstancedActors/InstancedActors.uplugin
@@ -26,21 +26,5 @@
 		}
 	],
 	"Plugins": [
-		{
-			"Name": "MassEntity",
-			"Enabled": true
-		},
-		{
-			"Name": "MassGameplay",
-			"Enabled": true
-		},
-		{
-			"Name": "MassAI",
-			"Enabled": true
-		},
-		{
-			"Name": "StructUtils",
-			"Enabled": true
-		},
 	]
 }
\ No newline at end of file
diff --git a/Engine/Plugins/Runtime/InstancedActors/Source/InstancedActors/InstancedActors.Build.cs b/Engine/Plugins/Runtime/InstancedActors/Source/InstancedActors/InstancedActors.Build.cs
--- a/Engine/Plugins/Runtime/InstancedActors/Source/InstancedActors/InstancedActors.Build.cs
+++ b/Engine/Plugins/Runtime/InstancedActors/Source/InstancedActors/InstancedActors.Build.cs
@@ -10,8 +10,6 @@ namespace UnrealBuildTool.Rules
 
 			PublicIncludePaths.AddRange(
 				new string[] {
-				"Runtime/AIModule/Public",
-				ModuleDirectory + "/Public",
 				}
 			);
 
@@ -25,14 +23,7 @@ namespace UnrealBuildTool.Rules
 				new string[] {
 					"Core",
 					"CoreUObject",
-					"Engine",
-					"MassEntity",
-					"MassActors",
-					"MassRepresentation",
-					"MassSignals",
-					"MassSpawner",
-					"MassLOD",
-					"StructUtils",
+					"Engine"
 				}
 			);
 
diff --git a/Engine/Plugins/Runtime/InstancedActors/Source/InstancedActorsTestSuite/InstancedActorsTestSuite.Build.cs b/Engine/Plugins/Runtime/InstancedActors/Source/InstancedActorsTestSuite/InstancedActorsTestSuite.Build.cs
--- a/Engine/Plugins/Runtime/InstancedActors/Source/InstancedActorsTestSuite/InstancedActorsTestSuite.Build.cs
+++ b/Engine/Plugins/Runtime/InstancedActors/Source/InstancedActorsTestSuite/InstancedActorsTestSuite.Build.cs
@@ -15,7 +15,6 @@ namespace UnrealBuildTool.Rules
 					"Core",
 					"CoreUObject",
 					"Engine",
-					"MassEntity",
 					"InstancedActors",
 					"AITestSuite",
 				}
