CIS  Issue 425138
Added missing "StructUtils" as a dependency to InstancedActors plugin

diff --git a/Engine/Plugins/Runtime/InstancedActors/Source/InstancedActors/InstancedActors.Build.cs b/Engine/Plugins/Runtime/InstancedActors/Source/InstancedActors/InstancedActors.Build.cs
--- a/Engine/Plugins/Runtime/InstancedActors/Source/InstancedActors/InstancedActors.Build.cs
+++ b/Engine/Plugins/Runtime/InstancedActors/Source/InstancedActors/InstancedActors.Build.cs
@@ -27,6 +27,7 @@ namespace UnrealBuildTool.Rules
 					"GameplayTags",
 					"MassEntity",
 					"MassRepresentation",
+					"StructUtils"
 				}
 			);
 
