Quick fix to avoid "Warning: Failed to find object" when calling UInstancedActorsProjectSettings::GetInstancedActorsSubsystemClass

diff --git a/Engine/Plugins/Runtime/InstancedActors/Source/InstancedActors/Private/InstancedActorsSettings.cpp b/Engine/Plugins/Runtime/InstancedActors/Source/InstancedActors/Private/InstancedActorsSettings.cpp
--- a/Engine/Plugins/Runtime/InstancedActors/Source/InstancedActors/Private/InstancedActorsSettings.cpp
+++ b/Engine/Plugins/Runtime/InstancedActors/Source/InstancedActors/Private/InstancedActorsSettings.cpp
@@ -29,8 +29,11 @@ TSubclassOf<UMassActorSpawnerSubsystem> UInstancedActorsProjectSettings::GetClie
 }
 
 TSubclassOf<UInstancedActorsSubsystem> UInstancedActorsProjectSettings::GetInstancedActorsSubsystemClass() const 
-{ 
-	return InstancedActorsSubsystemClass.TryLoadClass<UInstancedActorsSubsystem>(); 
+{
+	// @todo using ResolveClass rather than TryLoadClass since the latter emits "Warning: Failed to find object" if 
+	// the indicated class cannot be loaded. That can happen if the target class is loaded later with a GameplayFeatyre plugin
+	// This whole section needs changing so that different game modes can utilize different classes.
+	return InstancedActorsSubsystemClass.ResolveClass();
 }
 
 TSubclassOf<UMassStationaryDistanceVisualizationTrait> UInstancedActorsProjectSettings::GetStationaryVisualizationTraitClass() const
