Fix spelling error
#rnx

diff --git a/Engine/Plugins/Runtime/InstancedActors/Source/InstancedActors/Private/InstancedActorsRepresentationSubsystem.cpp b/Engine/Plugins/Runtime/InstancedActors/Source/InstancedActors/Private/InstancedActorsRepresentationSubsystem.cpp
--- a/Engine/Plugins/Runtime/InstancedActors/Source/InstancedActors/Private/InstancedActorsRepresentationSubsystem.cpp
+++ b/Engine/Plugins/Runtime/InstancedActors/Source/InstancedActors/Private/InstancedActorsRepresentationSubsystem.cpp
@@ -19,6 +19,6 @@ void UInstancedActorsRepresentationSubsystem::Initialize(FSubsystemCollectionBas
 		ActorSpawnerSubsystem = Cast<UMassActorSpawnerSubsystem>(Collection.InitializeDependency(GET_INSTANCEDACTORS_CONFIG_VALUE(GetClientActorSpawnerSubsystemClass())));
 	}
 
-	ensureMsgf(ActorSpawnerSubsystem, TEXT("Trying to initialize dependency on class %s failed. Verify InstanedActors settings.")
+	ensureMsgf(ActorSpawnerSubsystem, TEXT("Trying to initialize dependency on class %s failed. Verify InstancedActors settings.")
 		, *GetNameSafe(ActorSpawnerSubsystem));
 }
