Fix to editor-time UInstancedActorsSubsystem::DoesActorClassHaveRegisteredSettings that was using FInstancedActorsClassSettings rather than FInstancedActorsClassSettingsBase which caused issues when a project-specific FInstancedActorsClassSettingsBase extention was being used.
Also:
* Renamed FInstancedActorsClassSettingsBase::GetOverrideSettings to more precise MakeOverrideSettings
* Updated the comments erroneously pointing at FInstancedActorsClassSettings as the type to use.

diff --git a/Engine/Plugins/Runtime/InstancedActors/Source/InstancedActors/Private/InstancedActorsSubsystem.cpp b/Engine/Plugins/Runtime/InstancedActors/Source/InstancedActors/Private/InstancedActorsSubsystem.cpp
--- a/Engine/Plugins/Runtime/InstancedActors/Source/InstancedActors/Private/InstancedActorsSubsystem.cpp
+++ b/Engine/Plugins/Runtime/InstancedActors/Source/InstancedActors/Private/InstancedActorsSubsystem.cpp
@@ -53,7 +53,7 @@ namespace InstancedActorsCVars
 	FAutoConsoleVariableRef CVarRuntimeEnforceActorClassSettingsPresence(
 		TEXT("IA.RuntimeEnforceActorClassSettingsPresence"),
 		RuntimeEnforceActorClassSettingsPresence,
-		TEXT("The error severity to use when no FInstancedActorsClassSettings are found for a given ActorClass (or any of it's superclasses) in the ActorClassSettingsRegistry")
+		TEXT("The error severity to use when no FInstancedActorsClassSettingsBase are found for a given ActorClass (or any of it's superclasses) in the ActorClassSettingsRegistry")
 		TEXT("at runtime. Useful for ensuring unknown / unoptimized actor classes aren't being unexpectedly instanced.")
 		TEXT("0 = No error, ActorClass's are not required to be be present in ActorClassSettingsRegistry at all.")
 		TEXT("1 = Log an error, continue to instance ActorClass regardless.")
@@ -64,7 +64,7 @@ namespace InstancedActorsCVars
 	FAutoConsoleVariableRef CVarEditorEnforceActorClassSettingsPresence(
 		TEXT("IA.EditorEnforceActorClassSettingsPresence"),
 		EditorEnforceActorClassSettingsPresence,
-		TEXT("The error severity to use when no FInstancedActorsClassSettings are found for a given ActorClass (or any of it's superclasses) in the ActorClassSettingsRegistry")
+		TEXT("The error severity to use when no FInstancedActorsClassSettingsBase are found for a given ActorClass (or any of it's superclasses) in the ActorClassSettingsRegistry")
 		TEXT("when instancing actors in the editor. Useful for ensuring unknown / unoptimized actor classes aren't being unexpectedly instanced.")
 		TEXT("0 = No error, ActorClass's are not required to be be present in ActorClassSettingsRegistry at all.")
 		TEXT("1 = Log a message log warning, continue to instance ActorClass regardless.")
@@ -661,9 +661,9 @@ bool UInstancedActorsSubsystem::DoesActorClassHaveRegisteredSettings(TSubclassOf
 	UClass* ClassOrSuperClass = ActorClass.Get();
 	while (ClassOrSuperClass != nullptr)
 	{
-		// Find FInstancedActorsClassSettings for ClassOrSuperClass
+		// Find FInstancedActorsClassSettingsBase for ClassOrSuperClass
 		// Note: For fast lookup, we use the classes FName to lookup class settings, requiring class names to be unique for per-class settings
-		const FInstancedActorsClassSettings* ClassOrSuperClassSettings = DataRegistrySubsystem->GetCachedItem<FInstancedActorsClassSettings>({ProjectSettings->ActorClassSettingsRegistryType, ClassOrSuperClass->GetFName()});
+		const FInstancedActorsClassSettingsBase* ClassOrSuperClassSettings = DataRegistrySubsystem->GetCachedItem<FInstancedActorsClassSettingsBase>({ProjectSettings->ActorClassSettingsRegistryType, ClassOrSuperClass->GetFName()});
 		if (ClassOrSuperClassSettings != nullptr)
 		{
 			return true;
@@ -723,7 +723,7 @@ FSharedStruct UInstancedActorsSubsystem::CompileSettingsForActorClass(TSubclassO
 	UClass* ClassOrSuperClass = ActorClass.Get();
 	while (ClassOrSuperClass != nullptr)
 	{
-		// Find FInstancedActorsClassSettings for ClassOrSuperClass
+		// Find FInstancedActorsClassSettingsBase for ClassOrSuperClass
 		// Note: For fast lookup, we use the classes FName to lookup class settings, requiring class names to be unique for per-class settings
 		const FInstancedActorsClassSettingsBase* ClassOrSuperClassSettings = DataRegistrySubsystem->GetCachedItem<FInstancedActorsClassSettingsBase>({ProjectSettings->ActorClassSettingsRegistryType, ClassOrSuperClass->GetFName()});
 		if (ClassOrSuperClassSettings != nullptr)
@@ -731,14 +731,14 @@ FSharedStruct UInstancedActorsSubsystem::CompileSettingsForActorClass(TSubclassO
 			bFoundClassSettings = true;
 
 			// Apply class OverrideSettings
-			CompiledSettings.Get<FInstancedActorsSettings>().OverrideIfDefault(ClassOrSuperClassSettings->GetOverrideSettings(), ClassOrSuperClass->GetFName());
+			CompiledSettings.Get<FInstancedActorsSettings>().OverrideIfDefault(ClassOrSuperClassSettings->MakeOverrideSettings(), ClassOrSuperClass->GetFName());
 
 			// Apple class BaseSettings in reverse order
 			for (const FName& BaseSettingsName : ReverseIterate(ClassOrSuperClassSettings->BaseSettings))
 			{
 				const FInstancedStruct BaseSettings = GetCachedItem({ ProjectSettings->NamedSettingsRegistryType, BaseSettingsName });
 
-				if (ensureMsgf(BaseSettings.IsValid(), TEXT("FInstancedActorsClassSettings (%s) references unknown named settings '%s', skipping.")
+				if (ensureMsgf(BaseSettings.IsValid(), TEXT("FInstancedActorsClassSettingsBase (%s) references unknown named settings '%s', skipping.")
 					, *ClassOrSuperClass->GetPathName(), *BaseSettingsName.ToString()))
 				{
 					CompiledSettings.Get<FInstancedActorsSettings>().OverrideIfDefault(BaseSettings, BaseSettingsName);
diff --git a/Engine/Plugins/Runtime/InstancedActors/Source/InstancedActors/Public/InstancedActorsSettings.h b/Engine/Plugins/Runtime/InstancedActors/Source/InstancedActors/Public/InstancedActorsSettings.h
--- a/Engine/Plugins/Runtime/InstancedActors/Source/InstancedActors/Public/InstancedActorsSettings.h
+++ b/Engine/Plugins/Runtime/InstancedActors/Source/InstancedActors/Public/InstancedActorsSettings.h
@@ -16,7 +16,7 @@ class UMassStationaryDistanceVisualizationTrait;
 
 /** 
  * Configurable project settings for the Instanced Actors system.
- * @see FInstancedActorsClassSettings for per-class specific runtime settings.
+ * @see FInstancedActorsClassSettingsBase and FInstancedActorsClassSettings for per-class specific runtime settings.
  * @see AInstancedActorsManager
  */
 UCLASS(Config=InstancedActors, defaultconfig, DisplayName = "Instanced Actors")
@@ -40,7 +40,7 @@ public:
 	UPROPERTY(Config, EditAnywhere, BlueprintReadOnly, Category = ActorClassSettings)
 	FDataRegistryType NamedSettingsRegistryType = "InstancedActorsNamedSettings";
 
-	/** Data Registry to gather per-class FInstancedActorsClassSettings from during UInstancedActorsSubsystem init */
+	/** Data Registry to gather per-class FInstancedActorsClassSettingsBase-based settings from during UInstancedActorsSubsystem init */
 	UPROPERTY(Config, EditAnywhere, BlueprintReadOnly, Category = ActorClassSettings)
 	FDataRegistryType ActorClassSettingsRegistryType = "InstancedActorsClassSettings";
 
diff --git a/Engine/Plugins/Runtime/InstancedActors/Source/InstancedActors/Public/InstancedActorsSettingsTypes.h b/Engine/Plugins/Runtime/InstancedActors/Source/InstancedActors/Public/InstancedActorsSettingsTypes.h
--- a/Engine/Plugins/Runtime/InstancedActors/Source/InstancedActors/Public/InstancedActorsSettingsTypes.h
+++ b/Engine/Plugins/Runtime/InstancedActors/Source/InstancedActors/Public/InstancedActorsSettingsTypes.h
@@ -54,9 +54,9 @@
  * Named settings are FInstancedActorsSettings registered by name in UInstancedActorsProjectSettings::NamedSettingsRegistryType 
  * data registry.
  *
- * Class-specific settings are FInstancedActorsSettings registered by actor class via FInstancedActorsClassSettings's in the 
+ * Class-specific settings are FInstancedActorsSettings registered by actor class via FInstancedActorsClassSettingsBase's in the 
  * UInstancedActorsProjectSettings::ActorClassSettingsRegistryType data registry. Named settings can be used as bases here via
- * FInstancedActorsClassSettings::BaseSettings.
+ * FInstancedActorsClassSettingsBase::BaseSettings.
  *
  * @see UInstancedActorsProjectSettings for settings registry info.
  */
@@ -260,7 +260,7 @@ struct TStructOpsTypeTraits<FInstancedActorsSettings> : public TStructOpsTypeTra
  *    1) Default construted FInstancedActorsSettings
  *    2) UInstancedActorsProjectSettings::DefaultBaseSettingsName
  * 
- *    3) ActorClass::Super's FInstancedActorsClassSettings BaseSettings (if any)
+ *    3) ActorClass::Super's FInstancedActorsClassSettingsBase BaseSettings (if any)
  *         [0]
  *         [1]
  *         [2]
@@ -284,7 +284,7 @@ struct INSTANCEDACTORS_API FInstancedActorsClassSettingsBase : public FTableRowB
 {
 	GENERATED_BODY()
 	
-	// @There's FDataRegistryId (with UI), but I think it cannot be used here since it has both the registry and name innit.
+	// @todo There's FDataRegistryId (with UI), but I think it cannot be used here since it has both the registry and name init.
 	/** 
 	 * Optional ordered list of 'named' settings to apply to instances of ActorClass before applying OverrideSettings.
 	 * BaseSettings are applied in order, so the last setting wins (with OverrideSettings having final say).
@@ -293,10 +293,17 @@ struct INSTANCEDACTORS_API FInstancedActorsClassSettingsBase : public FTableRowB
 	UPROPERTY(EditAnywhere, BlueprintReadWrite, Category = InstancedActors)
 	TArray<FName> BaseSettings;
 
-	virtual FInstancedStruct GetOverrideSettings() const { return FInstancedStruct::Make<FInstancedActorsSettings>(); }
+	/** A virtual function to be overridden by child-types that provide an instance of their specific FInstancedActorsSettings flavor */
+	virtual FInstancedStruct MakeOverrideSettings() const { return FInstancedStruct::Make<FInstancedActorsSettings>(); }
 
 };
 
+/** 
+ * Generic implementation of FInstancedActorsClassSettingsBase that's using the generic FInstancedActorsSettings to override
+ * existing settings. If you want to extend the per-actor-class settings for your project then inhering either from 
+ * FInstancedActorsClassSettings and add extra logic and properties to the child type, or inherit from FInstancedActorsClassSettingsBase
+ * to provide a project-specific settings struct that inherits FInstancedActorsSettings.
+ */
 USTRUCT()
 struct INSTANCEDACTORS_API FInstancedActorsClassSettings : public FInstancedActorsClassSettingsBase
 {
@@ -309,5 +316,5 @@ struct INSTANCEDACTORS_API FInstancedActorsClassSettings : public FInstancedActo
 	UPROPERTY(EditAnywhere, BlueprintReadWrite, Category = InstancedActors)
 	FInstancedActorsSettings OverrideSettings;
 
-	virtual FInstancedStruct GetOverrideSettings() const { return FInstancedStruct::Make<FInstancedActorsSettings>(OverrideSettings); }
+	virtual FInstancedStruct MakeOverrideSettings() const { return FInstancedStruct::Make<FInstancedActorsSettings>(OverrideSettings); }
 };
diff --git a/Engine/Plugins/Runtime/InstancedActors/Source/InstancedActors/Public/InstancedActorsSubsystem.h b/Engine/Plugins/Runtime/InstancedActors/Source/InstancedActors/Public/InstancedActorsSubsystem.h
--- a/Engine/Plugins/Runtime/InstancedActors/Source/InstancedActors/Public/InstancedActorsSubsystem.h
+++ b/Engine/Plugins/Runtime/InstancedActors/Source/InstancedActors/Public/InstancedActorsSubsystem.h
@@ -110,18 +110,18 @@ public:
 	INSTANCEDACTORS_API AActor& GetOrCreateExemplarActor(TSubclassOf<AActor> ActorClass);
 
 	/** 
-	 * Compiles and caches finalized settings for ActorClass based off FInstancedActorsClassSettings found in 
+	 * Compiles and caches finalized settings for ActorClass based off FInstancedActorsClassSettingsBase found in 
 	 * UInstancedActorsProjectSettings::ActorClassSettingsRegistryType data registry, for ActorClass and it's 
 	 * inherited super classes.
 	 *
-	 * Note: FInstancedActorsClassSettings are indexed by class FName (*not* the full path) in the data registry, 
+	 * Note: FInstancedActorsClassSettingsBase are indexed by class FName (*not* the full path) in the data registry, 
 	 *       for quick lookup in CompileSettingsForActorClass. This means unique class names must be used for 
 	 *       per-class settings.
 	 */
 	FSharedStruct GetOrCompileSettingsForActorClass(TSubclassOf<AActor> ActorClass);
 
 	/**
-	 * Returns true if ActorClass has a matching FInstancedActorsClassSettings entry in 
+	 * Returns true if ActorClass has a matching FInstancedActorsClassSettingsBase entry in 
 	 * UInstancedActorsProjectSettings::ActorClassSettingsRegistryType data registry.
 	 * Note: This relies on the registry being loaded at the time of calling i.e: in editor the registry must be
 	 *       set to preload in editor.
@@ -194,12 +194,12 @@ protected:
 	TMap<TObjectKey<UInstancedActorsModifierVolumeComponent>, FBox> DebugModifierVolumeBounds;
 #endif
 
-	// Cached finalized / flattened FInstancedActorsClassSettings for GetOrCompileSettingsForActorClass requested ActorClass.
+	// Cached finalized / flattened FInstancedActorsClassSettingsBase for GetOrCompileSettingsForActorClass requested ActorClass.
 	// Built via CompileSettingsForActorClass and cached in GetOrCompileSettingsForActorClass
 	TMap<TWeakObjectPtr<UClass>, FSharedStruct> PerActorClassSettings;
 
 	// Called in GetOrCompileSettingsForActorClass to compile finalized settings for ActorClass based off 
-	// FInstancedActorsClassSettings found in UInstancedActorsProjectSettings::ActorClassSettingsRegistryType
+	// FInstancedActorsClassSettingsBase found in UInstancedActorsProjectSettings::ActorClassSettingsRegistryType
 	// data registry, for ActorClass and it's inherited super classes.
 	FSharedStruct CompileSettingsForActorClass(TSubclassOf<AActor> ActorClass) const;
 
