Reworked InstancedActors configuration approach.

There's the default config now, but external code can register overrides that upon registration and unregistration trigger "current config compilation" (stored in `UInstancedActorsProjectSettings.CompiledActiveConfig`) and that's what's being used at runtime.

The CL also provides a GameFeatureAction that can be used to override InstancedActors settings.

A couple of places in existing code had to be adjusted to use freshly implemented class/subsystem getter functions in UE::InstancedActors::Utils since fetching InstancedActors-related subsystems directly from the UWorld can cause confusion as there might be multiple subsystems of similar type (related by a parent class for example).

diff --git a/Engine/Plugins/Runtime/InstancedActors/InstancedActors.uplugin b/Engine/Plugins/Runtime/InstancedActors/InstancedActors.uplugin
--- a/Engine/Plugins/Runtime/InstancedActors/InstancedActors.uplugin
+++ b/Engine/Plugins/Runtime/InstancedActors/InstancedActors.uplugin
@@ -41,6 +41,10 @@
 		{
 			"Name": "DataRegistry",
 			"Enabled": true
+		},
+		{
+			"Name": "GameFeatures",
+			"Enabled": true
 		}
 	]
 }
\ No newline at end of file
diff --git a/Engine/Plugins/Runtime/InstancedActors/Source/InstancedActors/InstancedActors.Build.cs b/Engine/Plugins/Runtime/InstancedActors/Source/InstancedActors/InstancedActors.Build.cs
--- a/Engine/Plugins/Runtime/InstancedActors/Source/InstancedActors/InstancedActors.Build.cs
+++ b/Engine/Plugins/Runtime/InstancedActors/Source/InstancedActors/InstancedActors.Build.cs
@@ -26,6 +26,7 @@ namespace UnrealBuildTool.Rules
 					"DataRegistry",
 					"DeveloperSettings",
 					"NetCore",
+					"GameFeatures"
 				}
 			);
 
diff --git a/Engine/Plugins/Runtime/InstancedActors/Source/InstancedActors/Private/ClientInstancedActorsSpawnerSubsystem.cpp b/Engine/Plugins/Runtime/InstancedActors/Source/InstancedActors/Private/ClientInstancedActorsSpawnerSubsystem.cpp
--- a/Engine/Plugins/Runtime/InstancedActors/Source/InstancedActors/Private/ClientInstancedActorsSpawnerSubsystem.cpp
+++ b/Engine/Plugins/Runtime/InstancedActors/Source/InstancedActors/Private/ClientInstancedActorsSpawnerSubsystem.cpp
@@ -25,12 +25,6 @@ bool UClientInstancedActorsSpawnerSubsystem::ShouldCreateSubsystem(UObject* Oute
 		return false;
 	}
 
-	// do not instantiate if configured to use a different (sub)class
-	if (GET_INSTANCEDACTORS_CONFIG_VALUE(ClientActorSpawnerSubsystemClass) != GetClass())
-	{
-		return false;
-	}
-
 	// @todo Add support for non-replay NM_Standalone where we should use UServerInstancedActorsSpawnerSubsystem for 
 	// authoritative actor spawning.
 	UWorld* World = Cast<UWorld>(Outer);
diff --git a/Engine/Plugins/Runtime/InstancedActors/Source/InstancedActors/Private/GameFeatureAction_ConfigureInstancedActors.cpp b/Engine/Plugins/Runtime/InstancedActors/Source/InstancedActors/Private/GameFeatureAction_ConfigureInstancedActors.cpp
new file mode 100644
--- /dev/null
+++ b/Engine/Plugins/Runtime/InstancedActors/Source/InstancedActors/Private/GameFeatureAction_ConfigureInstancedActors.cpp
@@ -0,0 +1,21 @@
+// Copyright Epic Games, Inc. All Rights Reserved.
+
+#include "GameFeatureAction_ConfigureInstancedActors.h"
+#include "InstancedActorsSettings.h"
+#include "InstancedActorsSubsystem.h"
+#include "GameFeaturesSubsystem.h"
+
+
+void UGameFeatureAction_ConfigureInstancedActors::OnGameFeatureActivating(FGameFeatureActivatingContext& Context)
+{
+	Super::OnGameFeatureActivating(Context);
+
+	GetMutableDefault<UInstancedActorsProjectSettings>()->RegisterConfigOverride(*this, ConfigOverride);
+}
+
+void UGameFeatureAction_ConfigureInstancedActors::OnGameFeatureDeactivating(FGameFeatureDeactivatingContext& Context)
+{
+	GetMutableDefault<UInstancedActorsProjectSettings>()->UnregisterConfigOverride(*this);
+
+	Super::OnGameFeatureDeactivating(Context);
+}
diff --git a/Engine/Plugins/Runtime/InstancedActors/Source/InstancedActors/Private/InstancedActorsComponent.cpp b/Engine/Plugins/Runtime/InstancedActors/Source/InstancedActors/Private/InstancedActorsComponent.cpp
--- a/Engine/Plugins/Runtime/InstancedActors/Source/InstancedActors/Private/InstancedActorsComponent.cpp
+++ b/Engine/Plugins/Runtime/InstancedActors/Source/InstancedActors/Private/InstancedActorsComponent.cpp
@@ -4,8 +4,10 @@
 #include "InstancedActorsManager.h"
 #include "InstancedActorsData.h"
 #include "ServerInstancedActorsSpawnerSubsystem.h"
-
+#include "InstancedActorsSettings.h"
 #include "Net/UnrealNetwork.h"
+#include "Engine/World.h"
+
 
 UInstancedActorsComponent::UInstancedActorsComponent()
 {
@@ -53,7 +55,7 @@ void UInstancedActorsComponent::InitializeComponent()
 	UWorld* World = GetWorld();
 	if (World && World->GetNetMode() == NM_DedicatedServer)
 	{
-		UServerInstancedActorsSpawnerSubsystem* ServerInstancedActorSpawnerSubystem = World->GetSubsystem<UServerInstancedActorsSpawnerSubsystem>();
+		UServerInstancedActorsSpawnerSubsystem* ServerInstancedActorSpawnerSubystem = UE::InstancedActors::Utils::GetServerInstancedActorsSpawnerSubsystem(*World);
 		if (IsValid(ServerInstancedActorSpawnerSubystem))
 		{
 			ServerInstancedActorSpawnerSubystem->OnInstancedActorComponentInitialize(*this);
diff --git a/Engine/Plugins/Runtime/InstancedActors/Source/InstancedActors/Private/InstancedActorsData.cpp b/Engine/Plugins/Runtime/InstancedActors/Source/InstancedActors/Private/InstancedActorsData.cpp
--- a/Engine/Plugins/Runtime/InstancedActors/Source/InstancedActors/Private/InstancedActorsData.cpp
+++ b/Engine/Plugins/Runtime/InstancedActors/Source/InstancedActors/Private/InstancedActorsData.cpp
@@ -10,6 +10,7 @@
 #include "InstancedActorsSubsystem.h"
 #include "InstancedActorsVisualizationTrait.h"
 #include "InstancedActorsSettingsTypes.h"
+#include "InstancedActorsSettings.h"
 #include "UObject/ObjectSaveContext.h"
 #include "Algo/Count.h"
 #include "Algo/NoneOf.h"
@@ -163,7 +164,7 @@ void UInstancedActorsData::Initialize()
 	AInstancedActorsManager& Manager = GetManagerChecked();
 
 	// Get the settings setup nice and early.
-	UInstancedActorsSubsystem& InstancedActorSubsystem = UInstancedActorsSubsystem::GetChecked(GetManager());
+	UInstancedActorsSubsystem& InstancedActorSubsystem = Manager.GetInstancedActorSubsystemChecked();
 	SharedSettings = InstancedActorSubsystem.GetOrCompileSettingsForActorClass(ActorClass);
 	const FInstancedActorsSettings* Settings = GetSettingsPtr<const FInstancedActorsSettings>();
 
@@ -204,7 +205,7 @@ void UInstancedActorsData::Initialize()
 	}
 
 	// Get or create exemplar actor to derive entities from
-	const AActor& ExemplarActor = Manager.GetInstancedActorSubsystemChecked().GetOrCreateExemplarActor(ActorClass);
+	const AActor& ExemplarActor = InstancedActorSubsystem.GetOrCreateExemplarActor(ActorClass);
 
 	// Add default visualization at index 0
 	//FInstancedActorsVisualizationDesc DefaultVisualiation = FInstancedActorsVisualizationDesc::FromActor(ExemplarActor, &UE::InstancedActors::VisualizationDescrFromActorAdditionalSteps);
diff --git a/Engine/Plugins/Runtime/InstancedActors/Source/InstancedActors/Private/InstancedActorsManager.cpp b/Engine/Plugins/Runtime/InstancedActors/Source/InstancedActors/Private/InstancedActorsManager.cpp
--- a/Engine/Plugins/Runtime/InstancedActors/Source/InstancedActors/Private/InstancedActorsManager.cpp
+++ b/Engine/Plugins/Runtime/InstancedActors/Source/InstancedActors/Private/InstancedActorsManager.cpp
@@ -7,6 +7,7 @@
 #include "InstancedActorsIteration.h"
 #include "InstancedActorsModifierVolumeComponent.h"
 #include "InstancedActorsSettingsTypes.h"
+#include "InstancedActorsSettings.h"
 #include "InstancedActorsSubsystem.h"
 #include "InstancedActorsRepresentationActorManagement.h"
 #include "ActorPartition/ActorPartitionSubsystem.h"
@@ -169,7 +170,7 @@ void AInstancedActorsManager::BeginPlay()
 
 	// Register with IA subsystem if it's available already, otherwise the subsystem will
 	// collect this manager when it initializes later and call OnAddedToSubsystem
-	UInstancedActorsSubsystem* PreinitializedInstancedActorSubsystem = World->GetSubsystem<UInstancedActorsSubsystem>();
+	UInstancedActorsSubsystem* PreinitializedInstancedActorSubsystem = UE::InstancedActors::Utils::GetInstancedActorsSubsystem(*World);
 	if (PreinitializedInstancedActorSubsystem)
 	{
 		// Register manager with subsystem now, which will immediately call OnAddedToSubsystem
@@ -645,7 +646,7 @@ UInstancedActorsData& AInstancedActorsManager::GetOrCreateActorInstanceData(TSub
 	// Get or create exemplar actor to derive entities from
 	UWorld* World = GetWorld();
 	check(World);
-	UInstancedActorsSubsystem* EditorInstancedActorSubsystem = World->GetSubsystem<UInstancedActorsSubsystem>();
+	UInstancedActorsSubsystem* EditorInstancedActorSubsystem = UE::InstancedActors::Utils::GetInstancedActorsSubsystem(*World);
 	check(EditorInstancedActorSubsystem);
 	const AActor& ExemplarActor = EditorInstancedActorSubsystem->GetOrCreateExemplarActor(ActorClass);
 
diff --git a/Engine/Plugins/Runtime/InstancedActors/Source/InstancedActors/Private/InstancedActorsRepresentationSubsystem.cpp b/Engine/Plugins/Runtime/InstancedActors/Source/InstancedActors/Private/InstancedActorsRepresentationSubsystem.cpp
--- a/Engine/Plugins/Runtime/InstancedActors/Source/InstancedActors/Private/InstancedActorsRepresentationSubsystem.cpp
+++ b/Engine/Plugins/Runtime/InstancedActors/Source/InstancedActors/Private/InstancedActorsRepresentationSubsystem.cpp
@@ -2,35 +2,39 @@
 
 #include "InstancedActorsRepresentationSubsystem.h"
 #include "InstancedActorsSettings.h"
-#include "InstancedActorsTypes.h"
 
 
 void UInstancedActorsRepresentationSubsystem::Initialize(FSubsystemCollectionBase& Collection)
 {
 	Super::Initialize(Collection);
 
-	TSubclassOf<UMassActorSpawnerSubsystem> SpawnerSystemSubclass;
-	// @todo Add support for non-replay NM_Standalone where we should use UServerInstancedActorsSpawnerSubsystem for 
-	// authoritative actor spawning.
-	if (GetWorldRef().GetNetMode() == NM_DedicatedServer)
-	{
-		SpawnerSystemSubclass = GET_INSTANCEDACTORS_CONFIG_VALUE(GetServerActorSpawnerSubsystemClass());
-
-	}
-	else
-	{
-		SpawnerSystemSubclass = GET_INSTANCEDACTORS_CONFIG_VALUE(GetClientActorSpawnerSubsystemClass());
-	}
-
+	TSubclassOf<UMassActorSpawnerSubsystem> SpawnerSystemSubclass = UE::InstancedActors::Utils::DetermineActorSpawnerSubsystemClass(GetWorldRef());
 	if (SpawnerSystemSubclass)
 	{
 		ActorSpawnerSubsystem = Cast<UMassActorSpawnerSubsystem>(Collection.InitializeDependency(SpawnerSystemSubclass));
 
 		ensureMsgf(ActorSpawnerSubsystem, TEXT("Trying to initialize dependency on class %s failed. Verify InstanedActors settings.")
-			, *GetNameSafe(ActorSpawnerSubsystem));
+			, *SpawnerSystemSubclass->GetName());
 	}
-	else
+
+	OnSettingsChangedHandle = GET_INSTANCEDACTORS_CONFIG_VALUE(GetOnSettingsUpdated()).AddUObject(this, &UInstancedActorsRepresentationSubsystem::OnSettingsChanged);
+}
+
+void UInstancedActorsRepresentationSubsystem::Deinitialize()
+{
+	GET_INSTANCEDACTORS_CONFIG_VALUE(GetOnSettingsUpdated()).Remove(OnSettingsChangedHandle);
+	ActorSpawnerSubsystem = nullptr;
+
+	Super::Deinitialize();
+}
+
+void UInstancedActorsRepresentationSubsystem::OnSettingsChanged()
+{
+	if (UWorld* World = GetWorld())
 	{
-		UE_LOG(LogInstancedActors, Warning, TEXT("%hs failed to load SpawnerSystemSubclass"), __FUNCTION__);
+		ActorSpawnerSubsystem = UE::InstancedActors::Utils::GetActorSpawnerSubsystem(*World);
+		UE_CLOG(ActorSpawnerSubsystem == nullptr, LogInstancedActors, Warning
+			, TEXT("%s %hs failed to fetch ActorSpawnerSubsystem instance, class %s.")
+			, *GetName(), __FUNCTION__, *GetNameSafe(UE::InstancedActors::Utils::DetermineActorSpawnerSubsystemClass(*World)));
 	}
 }
diff --git a/Engine/Plugins/Runtime/InstancedActors/Source/InstancedActors/Private/InstancedActorsSettings.cpp b/Engine/Plugins/Runtime/InstancedActors/Source/InstancedActors/Private/InstancedActorsSettings.cpp
--- a/Engine/Plugins/Runtime/InstancedActors/Source/InstancedActors/Private/InstancedActorsSettings.cpp
+++ b/Engine/Plugins/Runtime/InstancedActors/Source/InstancedActors/Private/InstancedActorsSettings.cpp
@@ -7,37 +7,94 @@
 #include "InstancedActorsVisualizationTrait.h"
 
 
-// @todo In the code below using ResolveClass rather than TryLoadClass since the latter emits "Warning: Failed to find object"
-// if the indicated class cannot be loaded. That can happen if the target class is loaded later with a GameplayFeatyre plugin
-// This whole section needs changing so that different game modes can utilize different classes.
-
 //-----------------------------------------------------------------------------
 // UInstancedActorsProjectSettings
 //-----------------------------------------------------------------------------
 UInstancedActorsProjectSettings::UInstancedActorsProjectSettings()
 {
-	ServerActorSpawnerSubsystemClass = UServerInstancedActorsSpawnerSubsystem::StaticClass();
-	ClientActorSpawnerSubsystemClass = UClientInstancedActorsSpawnerSubsystem::StaticClass();
-	InstancedActorsSubsystemClass = UInstancedActorsSubsystem::StaticClass();
-	StationaryVisualizationTraitClass = UInstancedActorsVisualizationTrait::StaticClass();
+	DefaultConfig.ServerActorSpawnerSubsystemClass = UServerInstancedActorsSpawnerSubsystem::StaticClass();
+	DefaultConfig.ClientActorSpawnerSubsystemClass = UClientInstancedActorsSpawnerSubsystem::StaticClass();
+	DefaultConfig.InstancedActorsSubsystemClass = UInstancedActorsSubsystem::StaticClass();
+	DefaultConfig.StationaryVisualizationTraitClass = UInstancedActorsVisualizationTrait::StaticClass();
+
+	CompiledActiveConfig = DefaultConfig;
 }
 
 TSubclassOf<UMassActorSpawnerSubsystem> UInstancedActorsProjectSettings::GetServerActorSpawnerSubsystemClass() const 
 { 
-	return ServerActorSpawnerSubsystemClass.ResolveClass();
+	return CompiledActiveConfig.ServerActorSpawnerSubsystemClass;
 }
 
 TSubclassOf<UMassActorSpawnerSubsystem> UInstancedActorsProjectSettings::GetClientActorSpawnerSubsystemClass() const 
-{ 
-	return ClientActorSpawnerSubsystemClass.ResolveClass();
+{
+	return CompiledActiveConfig.ClientActorSpawnerSubsystemClass;
 }
 
 TSubclassOf<UInstancedActorsSubsystem> UInstancedActorsProjectSettings::GetInstancedActorsSubsystemClass() const 
 {
-	return InstancedActorsSubsystemClass.ResolveClass();
+	return CompiledActiveConfig.InstancedActorsSubsystemClass;
 }
 
 TSubclassOf<UMassStationaryDistanceVisualizationTrait> UInstancedActorsProjectSettings::GetStationaryVisualizationTraitClass() const
 {
-	return StationaryVisualizationTraitClass.ResolveClass();
+	return CompiledActiveConfig.StationaryVisualizationTraitClass;
 }
+
+void UInstancedActorsProjectSettings::RegisterConfigOverride(UObject& Owner, const FInstancedActorsConfig& Config)
+{
+	FClassConfigOverrideEntry* Entry = ClassConfigOverrides.FindByPredicate([OwnerKey = &Owner](const FClassConfigOverrideEntry& Entry)
+		{ 
+			return Entry.Owner == OwnerKey;
+		});
+	if (Entry)
+	{
+		Entry->ConfigOverride = Config;
+	}
+	else
+	{
+		ClassConfigOverrides.Add({ &Owner, Config });
+	}
+	CompileSettings();
+}
+
+void UInstancedActorsProjectSettings::UnregisterConfigOverride(UObject& Owner)
+{
+	const int32 NumRemoved = ClassConfigOverrides.RemoveAll([OwnerKey = &Owner](const FClassConfigOverrideEntry& Entry)
+		{
+			return Entry.Owner == OwnerKey;
+		});
+	
+	if (NumRemoved)
+	{
+		CompileSettings();
+	}
+}
+
+/** 
+ * A helper macro for applying overrides to the specified Property. Note that we iterate starting from the latest
+ * override and quit as soon as a valid override is found.
+ */
+#define APPLY_OVERRIDE(Config, Property) \
+	for (int32 OverrideIndex = ClassConfigOverrides.Num() - 1; OverrideIndex >=0; --OverrideIndex) \
+	{ \
+		const FClassConfigOverrideEntry& Entry = ClassConfigOverrides[OverrideIndex]; \
+		if (Entry.ConfigOverride.Property) \
+		{ \
+			Config.Property = Entry.ConfigOverride.Property; \
+			break; \
+		} \
+	} \
+
+void UInstancedActorsProjectSettings::CompileSettings()
+{
+	CompiledActiveConfig = DefaultConfig;
+
+	APPLY_OVERRIDE(CompiledActiveConfig, ServerActorSpawnerSubsystemClass);
+	APPLY_OVERRIDE(CompiledActiveConfig, ClientActorSpawnerSubsystemClass);
+	APPLY_OVERRIDE(CompiledActiveConfig, InstancedActorsSubsystemClass);
+	APPLY_OVERRIDE(CompiledActiveConfig, StationaryVisualizationTraitClass);
+
+	OnSettingsUpdated.Broadcast();
+}
+
+#undef APPLY_OVERRIDE
diff --git a/Engine/Plugins/Runtime/InstancedActors/Source/InstancedActors/Private/InstancedActorsSubsystem.cpp b/Engine/Plugins/Runtime/InstancedActors/Source/InstancedActors/Private/InstancedActorsSubsystem.cpp
--- a/Engine/Plugins/Runtime/InstancedActors/Source/InstancedActors/Private/InstancedActorsSubsystem.cpp
+++ b/Engine/Plugins/Runtime/InstancedActors/Source/InstancedActors/Private/InstancedActorsSubsystem.cpp
@@ -7,6 +7,7 @@
 #include "InstancedActorsDebug.h"
 #include "InstancedActorsData.h"
 #include "InstancedActorsSettingsTypes.h"
+#include "InstancedActorsSettings.h"
 #include "ActorPartition/ActorPartitionSubsystem.h"
 #include "Algo/Find.h"
 #include "DataRegistry.h"
@@ -115,12 +116,6 @@ bool UInstancedActorsSubsystem::ShouldCreateSubsystem(UObject* Outer) const
 		return false;
 	}
 
-	// do not instantiate if configured to use a different (sub)class
-	if (GET_INSTANCEDACTORS_CONFIG_VALUE(InstancedActorsSubsystemClass) != GetClass())
-	{
-		return false;
-	}
-
 	// UInstancedActorsSubsystem must always be present for editor worlds to allow for InstanceActor etc editor operations
 	UWorld* World = CastChecked<UWorld>(Outer);
 	if (World->WorldType == EWorldType::Editor)
diff --git a/Engine/Plugins/Runtime/InstancedActors/Source/InstancedActors/Private/InstancedActorsTypes.cpp b/Engine/Plugins/Runtime/InstancedActors/Source/InstancedActors/Private/InstancedActorsTypes.cpp
--- a/Engine/Plugins/Runtime/InstancedActors/Source/InstancedActors/Private/InstancedActorsTypes.cpp
+++ b/Engine/Plugins/Runtime/InstancedActors/Source/InstancedActors/Private/InstancedActorsTypes.cpp
@@ -3,10 +3,62 @@
 #include "InstancedActorsTypes.h"
 #include "Components/InstancedStaticMeshComponent.h"
 #include "Engine/StaticMesh.h"
-
+#include "Engine/World.h"
+#include "InstancedActorsSettings.h"
+#include "ServerInstancedActorsSpawnerSubsystem.h"
+#include "ClientInstancedActorsSpawnerSubsystem.h"
 
 DEFINE_LOG_CATEGORY(LogInstancedActors);
 
+namespace UE::InstancedActors::Utils
+{
+	TSubclassOf<UMassActorSpawnerSubsystem> DetermineActorSpawnerSubsystemClass(UWorld& World)
+	{
+		// @todo Add support for non-replay NM_Standalone where we should use UServerInstancedActorsSpawnerSubsystem for 
+		// authoritative actor spawning.
+		if (World.GetNetMode() == NM_DedicatedServer)
+		{
+			return GET_INSTANCEDACTORS_CONFIG_VALUE(GetServerActorSpawnerSubsystemClass());
+
+		}
+		return GET_INSTANCEDACTORS_CONFIG_VALUE(GetClientActorSpawnerSubsystemClass());
+	}
+
+	UServerInstancedActorsSpawnerSubsystem* GetServerInstancedActorsSpawnerSubsystem(UWorld& World)
+	{
+		TSubclassOf<UMassActorSpawnerSubsystem> SpawnerSubsystemClass = GET_INSTANCEDACTORS_CONFIG_VALUE(GetServerActorSpawnerSubsystemClass());
+		check(SpawnerSubsystemClass);
+		return Cast<UServerInstancedActorsSpawnerSubsystem>(World.GetSubsystemBase(SpawnerSubsystemClass));
+	}
+
+	UClientInstancedActorsSpawnerSubsystem* GetClientInstancedActorsSpawnerSubsystem(UWorld& World)
+	{
+		TSubclassOf<UMassActorSpawnerSubsystem> SpawnerSubsystemClass = GET_INSTANCEDACTORS_CONFIG_VALUE(GetServerActorSpawnerSubsystemClass());
+		check(SpawnerSubsystemClass);
+		return Cast<UClientInstancedActorsSpawnerSubsystem>(World.GetSubsystemBase(SpawnerSubsystemClass));
+	}
+
+	UMassActorSpawnerSubsystem* GetActorSpawnerSubsystem(UWorld& World)
+	{
+		if (World.GetNetMode() == NM_DedicatedServer)
+		{
+			return GetServerInstancedActorsSpawnerSubsystem(World);
+
+		}
+		return GetClientInstancedActorsSpawnerSubsystem(World);
+	}
+
+	UInstancedActorsSubsystem* GetInstancedActorsSubsystem(UWorld& World)
+	{
+		TSubclassOf<UInstancedActorsSubsystem> InstancedActorsSubsystemClass = GET_INSTANCEDACTORS_CONFIG_VALUE(GetInstancedActorsSubsystemClass());
+		check(InstancedActorsSubsystemClass);
+
+		UInstancedActorsSubsystem* InstancedActorSubsystem = Cast<UInstancedActorsSubsystem>(World.GetSubsystemBase(InstancedActorsSubsystemClass));
+
+		return InstancedActorSubsystem;
+	}
+} // namespace UE::InstancedActors::Utils
+
 //-----------------------------------------------------------------------------
 // FInstancedActorsTagSet
 //-----------------------------------------------------------------------------
diff --git a/Engine/Plugins/Runtime/InstancedActors/Source/InstancedActors/Private/ServerInstancedActorsSpawnerSubsystem.cpp b/Engine/Plugins/Runtime/InstancedActors/Source/InstancedActors/Private/ServerInstancedActorsSpawnerSubsystem.cpp
--- a/Engine/Plugins/Runtime/InstancedActors/Source/InstancedActors/Private/ServerInstancedActorsSpawnerSubsystem.cpp
+++ b/Engine/Plugins/Runtime/InstancedActors/Source/InstancedActors/Private/ServerInstancedActorsSpawnerSubsystem.cpp
@@ -17,12 +17,6 @@ bool UServerInstancedActorsSpawnerSubsystem::ShouldCreateSubsystem(UObject* Oute
 		return false;
 	}
 
-	// do not instantiate if configured to use a different (sub)class
-	if (GET_INSTANCEDACTORS_CONFIG_VALUE(ServerActorSpawnerSubsystemClass) != GetClass())
-	{
-		return false;
-	}
-
 	// @todo Add support for non-replay NM_Standalone where we should use UServerInstancedActorsSpawnerSubsystem for 
 	// authoritative actor spawning.
 	UWorld* World = Cast<UWorld>(Outer);
diff --git a/Engine/Plugins/Runtime/InstancedActors/Source/InstancedActors/Public/GameFeatureAction_ConfigureInstancedActors.h b/Engine/Plugins/Runtime/InstancedActors/Source/InstancedActors/Public/GameFeatureAction_ConfigureInstancedActors.h
new file mode 100644
--- /dev/null
+++ b/Engine/Plugins/Runtime/InstancedActors/Source/InstancedActors/Public/GameFeatureAction_ConfigureInstancedActors.h
@@ -0,0 +1,29 @@
+// Copyright Epic Games, Inc. All Rights Reserved.
+
+#pragma once
+
+#include "GameFeatureAction.h"
+#include "InstancedActorsSettings.h"
+#include "GameFeatureAction_ConfigureInstancedActors.generated.h"
+
+
+/** 
+ * GameplayFeature Action carrying overrides to InstancedActors settings
+ */
+UCLASS(meta = (DisplayName = "Configure InstancedActors"))
+class UGameFeatureAction_ConfigureInstancedActors : public UGameFeatureAction
+{
+	GENERATED_BODY()
+
+public:	
+	//~ Begin UGameFeatureAction interface
+	virtual void OnGameFeatureActivating(FGameFeatureActivatingContext& Context) override;
+	virtual void OnGameFeatureDeactivating(FGameFeatureDeactivatingContext& Context) override;
+	//~ End UGameFeatureAction interface
+
+private:
+
+	UPROPERTY(EditAnywhere, Category=InstancedActors)
+	FInstancedActorsConfig ConfigOverride;
+};
+
diff --git a/Engine/Plugins/Runtime/InstancedActors/Source/InstancedActors/Public/InstancedActorsRepresentationSubsystem.h b/Engine/Plugins/Runtime/InstancedActors/Source/InstancedActors/Public/InstancedActorsRepresentationSubsystem.h
--- a/Engine/Plugins/Runtime/InstancedActors/Source/InstancedActors/Public/InstancedActorsRepresentationSubsystem.h
+++ b/Engine/Plugins/Runtime/InstancedActors/Source/InstancedActors/Public/InstancedActorsRepresentationSubsystem.h
@@ -14,7 +14,12 @@ class UInstancedActorsRepresentationSubsystem : public UMassRepresentationSubsys
 protected:
 	//~ Begin USubsystem Overrides
 	virtual void Initialize(FSubsystemCollectionBase& Collection) override;
+	virtual void Deinitialize() override;
 	//~ End USubsystem Overrides
+
+	void OnSettingsChanged();
+
+	FDelegateHandle OnSettingsChangedHandle;
 };
 
 template<>
diff --git a/Engine/Plugins/Runtime/InstancedActors/Source/InstancedActors/Public/InstancedActorsSettings.h b/Engine/Plugins/Runtime/InstancedActors/Source/InstancedActors/Public/InstancedActorsSettings.h
--- a/Engine/Plugins/Runtime/InstancedActors/Source/InstancedActors/Public/InstancedActorsSettings.h
+++ b/Engine/Plugins/Runtime/InstancedActors/Source/InstancedActors/Public/InstancedActorsSettings.h
@@ -5,15 +5,43 @@
 #include "DataRegistryId.h"
 #include "Engine/DeveloperSettings.h"
 #include "UObject/SoftObjectPtr.h"
+#include "MassActorSpawnerSubsystem.h"
+#include "InstancedActorsSubsystem.h"
+#include "MassStationaryDistanceVisualizationTrait.h"
 #include "InstancedActorsSettings.generated.h"
 
 
-class UMassActorSpawnerSubsystem;
-class UInstancedActorsSubsystem;
-class UMassStationaryDistanceVisualizationTrait;
-
 #define GET_INSTANCEDACTORS_CONFIG_VALUE(a) (GetMutableDefault<UInstancedActorsProjectSettings>()->a)
 
+USTRUCT()
+struct FInstancedActorsConfig
+{
+	GENERATED_BODY()
+
+	UPROPERTY(Config, EditAnywhere, Category = Subsystems, NoClear, meta = (MetaClass = "/Script/MassActors.MassActorSpawnerSubsystem"))
+	TSubclassOf<UMassActorSpawnerSubsystem> ServerActorSpawnerSubsystemClass; 
+
+	UPROPERTY(Config, EditAnywhere, Category = Subsystems, NoClear, meta = (MetaClass = "/Script/MassActors.MassActorSpawnerSubsystem"))
+	TSubclassOf<UMassActorSpawnerSubsystem> ClientActorSpawnerSubsystemClass;
+
+	UPROPERTY(Config, EditAnywhere, Category = Subsystems, NoClear, meta = (MetaClass = "/Script/InstancedActors.InstancedActorsSubsystem"))
+	TSubclassOf<UInstancedActorsSubsystem> InstancedActorsSubsystemClass;
+
+	UPROPERTY(Config, EditAnywhere, Category = Subsystems, NoClear, meta = (MetaClass = "/Script/MassRepresentation.MassStationaryDistanceVisualizationTrait"))
+	TSubclassOf<UMassStationaryDistanceVisualizationTrait> StationaryVisualizationTraitClass;
+};
+
+USTRUCT()
+struct FClassConfigOverrideEntry
+{
+	GENERATED_BODY()
+
+	FObjectKey Owner;
+	
+	UPROPERTY()
+	FInstancedActorsConfig ConfigOverride;
+};
+
 /** 
  * Configurable project settings for the Instanced Actors system.
  * @see FInstancedActorsClassSettingsBase and FInstancedActorsClassSettings for per-class specific runtime settings.
@@ -25,6 +53,8 @@ class UInstancedActorsProjectSettings : public UDeveloperSettings
 	GENERATED_BODY()
 
 public:
+	DECLARE_MULTICAST_DELEGATE(FOnSettingsChanged);
+
 	UInstancedActorsProjectSettings();
 
 	TSubclassOf<UMassActorSpawnerSubsystem> GetServerActorSpawnerSubsystemClass() const;
@@ -32,6 +62,15 @@ public:
 	TSubclassOf<UInstancedActorsSubsystem> GetInstancedActorsSubsystemClass() const;
 	TSubclassOf<UMassStationaryDistanceVisualizationTrait> GetStationaryVisualizationTraitClass() const;
 
+	void RegisterConfigOverride(UObject& Owner, const FInstancedActorsConfig& Config);
+	void UnregisterConfigOverride(UObject& Owner);
+
+	FOnSettingsChanged& GetOnSettingsUpdated() { return OnSettingsUpdated; }
+
+protected:
+	void CompileSettings();
+
+public:
 	/** 3D grid size (distance along side) for partitioned instanced actor managers */
 	UPROPERTY(Config, EditAnywhere, BlueprintReadOnly, meta = (ClampMin="0", Units=cm), Category = Grid)
 	int32 GridSize = 24480;
@@ -55,18 +94,16 @@ public:
 	UPROPERTY(Config, EditAnywhere, BlueprintReadOnly, Category = ActorClassSettings)
 	FName EnforcedSettingsName = NAME_None;
 
-	// TSubclassOf<UMassActorSpawnerSubsystem> ServerActorSpawnerSubsystemClass;
-	UPROPERTY(Config, EditAnywhere, BlueprintReadOnly, Category = Subsystems, NoClear, meta = (MetaClass = "/Script/MassActors.MassActorSpawnerSubsystem"))
-	FSoftClassPath ServerActorSpawnerSubsystemClass;
-
-	// TSubclassOf<UMassActorSpawnerSubsystem> ClientActorSpawnerSubsystemClass;
-	UPROPERTY(Config, EditAnywhere, BlueprintReadOnly, Category = Subsystems, NoClear, meta = (MetaClass = "/Script/MassActors.MassActorSpawnerSubsystem"))
-	FSoftClassPath ClientActorSpawnerSubsystemClass;
+	UPROPERTY(Config, EditAnywhere, Category = ActorClassSettings)
+	FInstancedActorsConfig DefaultConfig;
 
-	// TSubclassOf<UInstancedActorsSubsystem> InstancedActorsSubsystemClass;
-	UPROPERTY(Config, EditAnywhere, BlueprintReadOnly, Category = Subsystems, NoClear, meta = (MetaClass = "/Script/InstancedActors.InstancedActorsSubsystem"))
-	FSoftClassPath InstancedActorsSubsystemClass;
+protected:
+	FOnSettingsChanged OnSettingsUpdated;
 
-	UPROPERTY(Config, EditAnywhere, BlueprintReadOnly, Category = Subsystems, NoClear, meta = (MetaClass = "/Script/MassRepresentation.MassStationaryDistanceVisualizationTrait"))
-	FSoftClassPath StationaryVisualizationTraitClass;
+	/** Represents the current config combining DefaultConfig and all registered ClassConfigOverrides */
+	UPROPERTY(Transient)
+	FInstancedActorsConfig CompiledActiveConfig;
+		
+	UPROPERTY(Transient)
+	TArray<FClassConfigOverrideEntry> ClassConfigOverrides;
 };
diff --git a/Engine/Plugins/Runtime/InstancedActors/Source/InstancedActors/Public/InstancedActorsSubsystem.h b/Engine/Plugins/Runtime/InstancedActors/Source/InstancedActors/Public/InstancedActorsSubsystem.h
--- a/Engine/Plugins/Runtime/InstancedActors/Source/InstancedActors/Public/InstancedActorsSubsystem.h
+++ b/Engine/Plugins/Runtime/InstancedActors/Source/InstancedActors/Public/InstancedActorsSubsystem.h
@@ -2,7 +2,6 @@
 
 #pragma once
 
-#include "InstancedActorsSettings.h"
 #include "InstancedActorsDebug.h"
 #include "InstancedActorsManager.h"
 #include "GameplayTagContainer.h"
@@ -21,6 +20,7 @@ class ULevel;
 struct FInstancedActorsInstanceHandle;
 struct FInstancedActorsManagerHandle;
 struct FInstancedActorsModifierVolumeHandle;
+class UInstancedActorsProjectSettings;
 
 /**
  * Instanced Actor subsystem used to spawn AInstancedActorsManager's and populate their instance data.
diff --git a/Engine/Plugins/Runtime/InstancedActors/Source/InstancedActors/Public/InstancedActorsTypes.h b/Engine/Plugins/Runtime/InstancedActors/Source/InstancedActors/Public/InstancedActorsTypes.h
--- a/Engine/Plugins/Runtime/InstancedActors/Source/InstancedActors/Public/InstancedActorsTypes.h
+++ b/Engine/Plugins/Runtime/InstancedActors/Source/InstancedActors/Public/InstancedActorsTypes.h
@@ -16,6 +16,11 @@ struct FStaticMeshInstanceVisualizationDesc;
 struct FStreamableHandle;
 struct FMassISMCSharedData;
 class UInstancedActorsData;
+class UWorld;
+class UServerInstancedActorsSpawnerSubsystem;
+class UClientInstancedActorsSpawnerSubsystem;
+class UMassActorSpawnerSubsystem;
+class UInstancedActorsSubsystem;
 
 enum class EInstancedActorsBulkLOD : uint8
 {
@@ -37,6 +42,18 @@ enum class EInstancedActorsFragmentFlags : uint8
 };
 ENUM_CLASS_FLAGS(EInstancedActorsFragmentFlags);
 
+namespace UE::InstancedActors::Utils
+{
+	INSTANCEDACTORS_API TSubclassOf<UMassActorSpawnerSubsystem> DetermineActorSpawnerSubsystemClass(UWorld& World);
+	INSTANCEDACTORS_API UServerInstancedActorsSpawnerSubsystem* GetServerInstancedActorsSpawnerSubsystem(UWorld& World);
+	INSTANCEDACTORS_API UClientInstancedActorsSpawnerSubsystem* GetClientInstancedActorsSpawnerSubsystem(UWorld& World);
+	/** 
+	 * Calls either GetServerInstancedActorsSpawnerSubsystem or GetClientInstancedActorsSpawnerSubsystem, depending on 
+	 * given UWorld's net mode.
+	 */
+	INSTANCEDACTORS_API UMassActorSpawnerSubsystem* GetActorSpawnerSubsystem(UWorld& World);
+	INSTANCEDACTORS_API UInstancedActorsSubsystem* GetInstancedActorsSubsystem(UWorld& World);
+}
 
 // FInstancedActorsTagSet -> FInstancedActorsTagSet
 /** An immutable hashed tag container used to categorize / partition instances */
