Fixed a dumb mistake in UE::InstancedActors::Utils::GetClientInstancedActorsSpawnerSubsystem resulting in trying to fetch server-side classs while running for client.

diff --git a/Engine/Plugins/Runtime/InstancedActors/Source/InstancedActors/Private/InstancedActorsTypes.cpp b/Engine/Plugins/Runtime/InstancedActors/Source/InstancedActors/Private/InstancedActorsTypes.cpp
--- a/Engine/Plugins/Runtime/InstancedActors/Source/InstancedActors/Private/InstancedActorsTypes.cpp
+++ b/Engine/Plugins/Runtime/InstancedActors/Source/InstancedActors/Private/InstancedActorsTypes.cpp
@@ -33,7 +33,7 @@ namespace UE::InstancedActors::Utils
 
 	UClientInstancedActorsSpawnerSubsystem* GetClientInstancedActorsSpawnerSubsystem(UWorld& World)
 	{
-		TSubclassOf<UMassActorSpawnerSubsystem> SpawnerSubsystemClass = GET_INSTANCEDACTORS_CONFIG_VALUE(GetServerActorSpawnerSubsystemClass());
+		TSubclassOf<UMassActorSpawnerSubsystem> SpawnerSubsystemClass = GET_INSTANCEDACTORS_CONFIG_VALUE(GetClientActorSpawnerSubsystemClass());
 		check(SpawnerSubsystemClass);
 		return Cast<UClientInstancedActorsSpawnerSubsystem>(World.GetSubsystemBase(SpawnerSubsystemClass));
 	}
