Removed UInstancedActorsComponent::OnClientRegisteredForInstance as it's never being called

diff --git a/Engine/Plugins/Runtime/InstancedActors/Source/InstancedActors/Private/InstancedActorsComponent.cpp b/Engine/Plugins/Runtime/InstancedActors/Source/InstancedActors/Private/InstancedActorsComponent.cpp
--- a/Engine/Plugins/Runtime/InstancedActors/Source/InstancedActors/Private/InstancedActorsComponent.cpp
+++ b/Engine/Plugins/Runtime/InstancedActors/Source/InstancedActors/Private/InstancedActorsComponent.cpp
@@ -24,11 +24,6 @@ void UInstancedActorsComponent::InitializeComponentForInstance(FInstancedActorsI
 	InstanceHandle = InInstanceHandle;
 }
 
-void UInstancedActorsComponent::OnClientRegisteredForInstance(FInstancedActorsInstanceHandle InInstanceHandle)
-{
-	InstanceHandle = InInstanceHandle;
-}
-
 void UInstancedActorsComponent::GetLifetimeReplicatedProps(TArray<FLifetimeProperty>& OutLifetimeProps) const
 {
 	Super::GetLifetimeReplicatedProps(OutLifetimeProps);
diff --git a/Engine/Plugins/Runtime/InstancedActors/Source/InstancedActors/Public/InstancedActorsComponent.h b/Engine/Plugins/Runtime/InstancedActors/Source/InstancedActors/Public/InstancedActorsComponent.h
--- a/Engine/Plugins/Runtime/InstancedActors/Source/InstancedActors/Public/InstancedActorsComponent.h
+++ b/Engine/Plugins/Runtime/InstancedActors/Source/InstancedActors/Public/InstancedActorsComponent.h
@@ -45,9 +45,6 @@ public:
 	 */
 	virtual void InitializeComponentForInstance(FInstancedActorsInstanceHandle InInstanceHandle);
 
-	/** Called on clients just prior to registering a replicated Instanced Actor with it's matching Mass Entity */
-	virtual void OnClientRegisteredForInstance(FInstancedActorsInstanceHandle InInstanceHandle);
-
 	/** 
 	 * Called on an 'exemplar' Actor's components for clients & servers during UInstancedActorsData::CreateEntityTemplate to provide 
 	 * UInstancedActorsComponent's an opportunity to extend Mass entity default traits.
