Fixed InstancedActorManager issues when its being registered with InstancedActorSubsystem multiple times. This can happen if a game has a project-specific InstancedActorsSubsystem override that gets created late due to being added via game feature actions.

diff --git a/Engine/Plugins/Runtime/InstancedActors/Source/InstancedActors/Private/InstancedActorsManager.cpp b/Engine/Plugins/Runtime/InstancedActors/Source/InstancedActors/Private/InstancedActorsManager.cpp
--- a/Engine/Plugins/Runtime/InstancedActors/Source/InstancedActors/Private/InstancedActorsManager.cpp
+++ b/Engine/Plugins/Runtime/InstancedActors/Source/InstancedActors/Private/InstancedActorsManager.cpp
@@ -180,7 +180,16 @@ void AInstancedActorsManager::BeginPlay()
 
 void AInstancedActorsManager::OnAddedToSubsystem(UInstancedActorsSubsystem& InInstancedActorSubsystem, FInstancedActorsManagerHandle InManagerHandle)
 {
-	checkf(!IsValid(InstancedActorSubsystem), TEXT("Manager %s has already been added to a UInstancedActorsSubsystem"), *GetPathName());
+	if (IsValid(InstancedActorSubsystem) && (InstancedActorSubsystem != &InInstancedActorSubsystem))
+	{
+		// we need to unregister from the previous subsystem first
+		DespawnAllEntities();
+		InstancedActorSubsystem->RemoveManager(ManagerHandle);
+		InstancedActorSubsystem = nullptr;
+		ManagerHandle.Reset();
+	}
+	
+	checkf(!IsValid(InstancedActorSubsystem), TEXT("Manager %s has already been added to a %s"), *GetPathName(), *InstancedActorSubsystem->GetName());
 	checkf(!ManagerHandle.IsValid(), TEXT("Manager %s has already been added to a UInstancedActorsSubsystem"), *GetPathName());
 	check(InManagerHandle.IsValid());
 
@@ -192,7 +201,8 @@ void AInstancedActorsManager::OnAddedToSubsystem(UInstancedActorsSubsystem& InIn
 	InstancedActorSubsystem->ForEachModifierVolume(InstanceBounds, [this](UInstancedActorsModifierVolumeComponent& ModifierVolume)
 		{
 			AddModifierVolume(ModifierVolume);
-			return true; });
+			return true; 
+		});
 
 	if (UE::InstancedActors::CVars::bDeferSpawnEntities)
 	{
@@ -250,16 +260,7 @@ void AInstancedActorsManager::InitializeModifyAndSpawnEntities()
 
 void AInstancedActorsManager::EndPlay(const EEndPlayReason::Type EndPlayReason)
 {
-	// DespawnEntities for all PerActorClassInstanceData
-	if (HasSpawnedEntities())
-	{
-		for (TObjectPtr<UInstancedActorsData> InstanceData : PerActorClassInstanceData)
-		{
-			// Reconstruct instance data from Mass then destroy all Mass entities
-			InstanceData->DespawnEntities();
-		}
-		bHasSpawnedEntities = false;
-	}
+	DespawnAllEntities();
 
 	// Deregister with UInstancedActorsSubsystem
 	if (InstancedActorSubsystem)
@@ -737,6 +738,20 @@ void AInstancedActorsManager::RuntimeRemoveAllInstances()
 	}
 }
 
+void AInstancedActorsManager::DespawnAllEntities()
+{
+	// DespawnEntities for all PerActorClassInstanceData
+	if (HasSpawnedEntities())
+	{
+		for (TObjectPtr<UInstancedActorsData> InstanceData : PerActorClassInstanceData)
+		{
+			// Reconstruct instance data from Mass then destroy all Mass entities
+			InstanceData->DespawnEntities();
+		}
+		bHasSpawnedEntities = false;
+	}
+}
+
 bool AInstancedActorsManager::ForEachInstance(FInstanceOperationFunc Operation) const
 {
 	FScopedInstancedActorsIterationContext IterationContext;
diff --git a/Engine/Plugins/Runtime/InstancedActors/Source/InstancedActors/Public/InstancedActorsManager.h b/Engine/Plugins/Runtime/InstancedActors/Source/InstancedActors/Public/InstancedActorsManager.h
--- a/Engine/Plugins/Runtime/InstancedActors/Source/InstancedActors/Public/InstancedActorsManager.h
+++ b/Engine/Plugins/Runtime/InstancedActors/Source/InstancedActors/Public/InstancedActorsManager.h
@@ -240,6 +240,9 @@ protected:
 	// @param TimeDelta		Real time in seconds since serialization (0 when saving)
 	void SerializeInstancePersistenceData(FStructuredArchive::FRecord Record, UInstancedActorsData* InstanceData, int64 TimeDelta) const;
 
+	// Despawns all entities spawned by individual UInstancedActorsData instances. 
+	virtual void DespawnAllEntities();
+
 	// Attempts to run any 'pending' modifiers in ModifierVolumes where are appropriate to run givem HasSpawnedEntities
 	// Called in BeginPlay prior to, and then again after SpawnEntities. Also called in AddModifierVolume.
 	// @see UInstancedActorsModifierBase::bRequiresSpawnedEntities
