Cleaning up registered UInstancedActorsModifierVolumeComponent instances upon InstancedActorsManager's re-adding to a InstancedActorsSubsystem

diff --git a/Engine/Plugins/Runtime/InstancedActors/Source/InstancedActors/Private/InstancedActorsManager.cpp b/Engine/Plugins/Runtime/InstancedActors/Source/InstancedActors/Private/InstancedActorsManager.cpp
--- a/Engine/Plugins/Runtime/InstancedActors/Source/InstancedActors/Private/InstancedActorsManager.cpp
+++ b/Engine/Plugins/Runtime/InstancedActors/Source/InstancedActors/Private/InstancedActorsManager.cpp
@@ -184,6 +184,9 @@ void AInstancedActorsManager::OnAddedToSubsystem(UInstancedActorsSubsystem& InIn
 	{
 		// we need to unregister from the previous subsystem first
 		DespawnAllEntities();
+		// and let registered modifiers go
+		RemoveAllModifierVolumes();
+
 		InstancedActorSubsystem->RemoveManager(ManagerHandle);
 		InstancedActorSubsystem = nullptr;
 		ManagerHandle.Reset();
@@ -1149,6 +1152,21 @@ void AInstancedActorsManager::RemoveModifierVolume(UInstancedActorsModifierVolum
 	ModifierVolume.OnRemovedFromManager(*this);
 }
 
+void AInstancedActorsManager::RemoveAllModifierVolumes()
+{
+	for (TWeakObjectPtr<UInstancedActorsModifierVolumeComponent> WeakVolumeComponent : ModifierVolumes)
+	{
+		if (UInstancedActorsModifierVolumeComponent* VolumeComponent = WeakVolumeComponent.Get())
+		{
+			VolumeComponent->OnRemovedFromManager(*this);
+		}
+	}
+	
+	ModifierVolumes.Reset();
+	PendingModifierVolumes.Reset();
+	PendingModifierVolumeModifiers.Reset();
+}
+
 void AInstancedActorsManager::TryRunPendingModifiers()
 {
 	for (TBitArray<>::FIterator PendingModifierVolumeIt(PendingModifierVolumes); PendingModifierVolumeIt; ++PendingModifierVolumeIt)
diff --git a/Engine/Plugins/Runtime/InstancedActors/Source/InstancedActors/Public/InstancedActorsManager.h b/Engine/Plugins/Runtime/InstancedActors/Source/InstancedActors/Public/InstancedActorsManager.h
--- a/Engine/Plugins/Runtime/InstancedActors/Source/InstancedActors/Public/InstancedActorsManager.h
+++ b/Engine/Plugins/Runtime/InstancedActors/Source/InstancedActors/Public/InstancedActorsManager.h
@@ -158,6 +158,7 @@ public:
 
 	void AddModifierVolume(UInstancedActorsModifierVolumeComponent& ModifierVolume);
 	void RemoveModifierVolume(UInstancedActorsModifierVolumeComponent& ModifierVolume);
+	void RemoveAllModifierVolumes();
 
 	// Request the persistant data system to re-save this managers persistent data
 	void RequestPersistentDataSave();
