Iris replication support for UInstancedActorsData.

diff --git a/Engine/Plugins/Runtime/InstancedActors/Source/InstancedActors/InstancedActors.Build.cs b/Engine/Plugins/Runtime/InstancedActors/Source/InstancedActors/InstancedActors.Build.cs
--- a/Engine/Plugins/Runtime/InstancedActors/Source/InstancedActors/InstancedActors.Build.cs
+++ b/Engine/Plugins/Runtime/InstancedActors/Source/InstancedActors/InstancedActors.Build.cs
@@ -46,6 +46,7 @@ namespace UnrealBuildTool.Rules
 			}
 
 			SetupGameplayDebuggerSupport(Target);
+			SetupIrisSupport(Target);
 		}
 	}
 }
diff --git a/Engine/Plugins/Runtime/InstancedActors/Source/InstancedActors/Private/InstancedActorsData.cpp b/Engine/Plugins/Runtime/InstancedActors/Source/InstancedActors/Private/InstancedActorsData.cpp
--- a/Engine/Plugins/Runtime/InstancedActors/Source/InstancedActors/Private/InstancedActorsData.cpp
+++ b/Engine/Plugins/Runtime/InstancedActors/Source/InstancedActors/Private/InstancedActorsData.cpp
@@ -17,6 +17,9 @@
 #include "Components/InstancedStaticMeshComponent.h"
 #include "Engine/AssetManager.h"
 #include "Engine/StaticMesh.h"
+#if UE_WITH_IRIS
+#include "Iris/ReplicationSystem/ReplicationFragmentUtil.h"
+#endif
 #include "MassActorSubsystem.h"
 #include "MassCommonFragments.h"
 #include "MassEntityConfigAsset.h"
@@ -865,6 +868,13 @@ void UInstancedActorsData::GetLifetimeReplicatedProps(TArray<FLifetimeProperty>&
 	DOREPLIFETIME_WITH_PARAMS_FAST(UInstancedActorsData, InstanceDeltas, RepParams);
 }
 
+#if UE_WITH_IRIS
+void UInstancedActorsData::RegisterReplicationFragments(UE::Net::FFragmentRegistrationContext& Context, UE::Net::EFragmentRegistrationFlags RegistrationFlags)
+{
+	UE::Net::FReplicationFragmentUtil::CreateAndRegisterFragmentsForObject(this, Context, RegistrationFlags);
+}
+#endif
+
 void UInstancedActorsData::SetInstanceCurrentLifecyclePhase(FInstancedActorsInstanceIndex InstanceIndex, uint8 InCurrentLifecyclePhaseIndex)
 {
 	AInstancedActorsManager& Manager = GetManagerChecked();
diff --git a/Engine/Plugins/Runtime/InstancedActors/Source/InstancedActors/Public/InstancedActorsData.h b/Engine/Plugins/Runtime/InstancedActors/Source/InstancedActors/Public/InstancedActorsData.h
--- a/Engine/Plugins/Runtime/InstancedActors/Source/InstancedActors/Public/InstancedActorsData.h
+++ b/Engine/Plugins/Runtime/InstancedActors/Source/InstancedActors/Public/InstancedActorsData.h
@@ -361,6 +361,9 @@ protected:
 	virtual void PostDuplicate(EDuplicateMode::Type DuplicateMode) override;
 	virtual void PostLoad() override;
 	virtual void GetLifetimeReplicatedProps(TArray<FLifetimeProperty>& OutLifetimeProps) const override;
+#if UE_WITH_IRIS
+	virtual void RegisterReplicationFragments(UE::Net::FFragmentRegistrationContext& Context, UE::Net::EFragmentRegistrationFlags RegistrationFlags) override;
+#endif
 	virtual bool IsNameStableForNetworking() const override { return true; }
 	//~ End UObject Overrides
 
