Made UInstancedActorsData class spawned by AInstancedActorsManager configurable

diff --git a/Engine/Plugins/Runtime/InstancedActors/Source/InstancedActors/Private/InstancedActorsManager.cpp b/Engine/Plugins/Runtime/InstancedActors/Source/InstancedActors/Private/InstancedActorsManager.cpp
--- a/Engine/Plugins/Runtime/InstancedActors/Source/InstancedActors/Private/InstancedActorsManager.cpp
+++ b/Engine/Plugins/Runtime/InstancedActors/Source/InstancedActors/Private/InstancedActorsManager.cpp
@@ -137,6 +137,7 @@ namespace UE::InstancedActors
 //-----------------------------------------------------------------------------
 AInstancedActorsManager::AInstancedActorsManager()
 {
+	InstancedActorsDataClass = UInstancedActorsData::StaticClass();
 	bReplicateUsingRegisteredSubObjectList = true;
 	bReplicates = true;
 	SetNetDormancy(DORM_DormantAll);
@@ -620,6 +621,23 @@ void AInstancedActorsManager::SetGridGuid(const FGuid& InGuid)
 	ManagerGridGuid = InGuid;
 }
 
+UInstancedActorsData* AInstancedActorsManager::CreateNextInstanceActorData(TSubclassOf<AActor> ActorClass, const FInstancedActorsTagSet& InstanceTags)
+{
+	check(InstancedActorsDataClass);
+
+	UInstancedActorsData* NewInstanceData = NewObject<UInstancedActorsData>(this, InstancedActorsDataClass);
+	// @todo it's conceivable the NextInstanceDataID will overflow. We need to use some handle system in place instead. 
+	NewInstanceData->ID = NextInstanceDataID++;
+	NewInstanceData->ActorClass = ActorClass;
+	NewInstanceData->Tags = InstanceTags;
+	check(Algo::NoneOf(PerActorClassInstanceData, [NewInstanceData](UInstancedActorsData* InstanceData)
+		{
+			return InstanceData->ID == NewInstanceData->ID;
+		}));
+
+	return NewInstanceData;
+}
+
 UInstancedActorsData& AInstancedActorsManager::GetOrCreateActorInstanceData(TSubclassOf<AActor> ActorClass, const FInstancedActorsTagSet& InstanceTags)
 {
 	checkf(HasActorBegunPlay() == false, TEXT("AInstancedActorsManager doesn't yet support runtime addition of instances"));
@@ -635,16 +653,7 @@ UInstancedActorsData& AInstancedActorsManager::GetOrCreateActorInstanceData(TSub
 		return **InstanceData;
 	}
 
-	UInstancedActorsData* NewInstanceData = NewObject<UInstancedActorsData>(this);
-	// @todo it's conceivable the NextInstanceDataID will overflow. We need to use some handle system in place instead. 
-	NewInstanceData->ID = NextInstanceDataID++;
-	NewInstanceData->ActorClass = ActorClass;
-	NewInstanceData->Tags = InstanceTags;
-	check(Algo::NoneOf(PerActorClassInstanceData, [NewInstanceData](UInstancedActorsData* InstanceData)
-		{ 
-			return InstanceData->ID == NewInstanceData->ID; 
-		}));
-
+	UInstancedActorsData* NewInstanceData = CreateNextInstanceActorData(ActorClass, InstanceTags);
 	PerActorClassInstanceData.Add(NewInstanceData);
 
 	// Get or create exemplar actor to derive entities from
diff --git a/Engine/Plugins/Runtime/InstancedActors/Source/InstancedActors/Public/InstancedActorsManager.h b/Engine/Plugins/Runtime/InstancedActors/Source/InstancedActors/Public/InstancedActorsManager.h
--- a/Engine/Plugins/Runtime/InstancedActors/Source/InstancedActors/Public/InstancedActorsManager.h
+++ b/Engine/Plugins/Runtime/InstancedActors/Source/InstancedActors/Public/InstancedActorsManager.h
@@ -259,6 +259,8 @@ protected:
 	// Helper function to create and initialize per-actor-class UInstancedActorsData's, optionally further partitioned by InstanceTags
 	UInstancedActorsData& GetOrCreateActorInstanceData(TSubclassOf<AActor> ActorClass, const FInstancedActorsTagSet& InstanceTags);
 
+	virtual UInstancedActorsData* CreateNextInstanceActorData(TSubclassOf<AActor> ActorClass, const FInstancedActorsTagSet& InstanceTags);
+
 	// Used to set the right properties on the editor ISMCs so we can do per-instance selection.
 	virtual void PreRegisterAllComponents() override;
 #endif
@@ -312,6 +314,10 @@ protected:
 	UPROPERTY(Transient)
 	TMap<TObjectPtr<UInstancedStaticMeshComponent>, int32> ISMComponentToInstanceDataMap;
 
+	/** Class to be spawned to represent individual actor class instances. */
+	UPROPERTY(EditAnywhere, Category=InstancedActor)
+	TSubclassOf<UInstancedActorsData> InstancedActorsDataClass;
+
 private:
 #if WITH_EDITORONLY_DATA
 	UPROPERTY()
