Resetting `NextTickTime` for InstancedActors shared fragments that are being reused.

diff --git a/Engine/Plugins/Runtime/InstancedActors/Source/InstancedActors/Private/InstancedActorsData.cpp b/Engine/Plugins/Runtime/InstancedActors/Source/InstancedActors/Private/InstancedActorsData.cpp
--- a/Engine/Plugins/Runtime/InstancedActors/Source/InstancedActors/Private/InstancedActorsData.cpp
+++ b/Engine/Plugins/Runtime/InstancedActors/Source/InstancedActors/Private/InstancedActorsData.cpp
@@ -144,6 +144,13 @@ namespace UE::InstancedActors
 			TEXT("If enabled, forces all calls to UInstancedActorsData::AddVisualizationAsync to immediately sync load required assets and ")
 			TEXT("initialize the visualization, instead of requesting async loads and deferring initialization."),
 			ECVF_Default);
+
+		bool bUpdateNextTickTimeFragments = true;
+		FAutoConsoleVariableRef CVarUpdateNextTickTimeFragments(
+			TEXT("IA.UpdateNextTickTimeFragments"),
+			bUpdateNextTickTimeFragments,
+			TEXT("If enabled (default) forces LOD updated for newly re-spawned entities (i.e. spawned by IAD that has been reloaded)"),
+			ECVF_Default);
 	} // CVars
 
 	namespace Helpers
@@ -318,6 +325,19 @@ void UInstancedActorsData::SpawnEntities()
 
 	// Now we've seeded Mass entity locations, we can free up now-superfluous InstanceTransforms
 	InstanceTransforms.Empty();
+
+	if (UE::InstancedActors::CVars::bUpdateNextTickTimeFragments)
+	{
+		FInstancedActorsDataSharedFragment* AsShared = SharedInstancedActorDataStruct.GetPtr<FInstancedActorsDataSharedFragment>();
+		// our shared fragment was being ticked in the past, so it's possible it's scheduled to tick in quite some time
+		// while we need the update ASAP. We're letting the UInstancedActorsSubsystem know to reschedule it.
+		if (ensure(AsShared) && AsShared->LastTickTime > 0.0)
+		{
+			UInstancedActorsSubsystem* InstancedActorSubsystem = UE::InstancedActors::Utils::GetInstancedActorsSubsystem(*World);
+			check(InstancedActorSubsystem);
+			InstancedActorSubsystem->UpdateAndResetTickTime(SharedInstancedActorDataStruct);
+		}
+	}
 }
 
 void UInstancedActorsData::DespawnEntities()
@@ -640,10 +660,12 @@ FMassEntityHandle UInstancedActorsData::GetEntity(FInstancedActorsInstanceIndex
 
 void UInstancedActorsData::SetSharedInstancedActorDataStruct(FSharedStruct InSharedStruct)
 {
-	checkf(SharedInstancedActorDataStruct.IsValid() == false || SharedInstancedActorDataStruct.Identical(&InSharedStruct, 0), TEXT("We don't expect to override this value"));
 	checkf(InSharedStruct.GetPtr<FInstancedActorsDataSharedFragment>(), TEXT("We expect only FInstancedActorsDataSharedFragment-base types here"));
+	checkf(SharedInstancedActorDataStruct.IsValid() == false 
+		|| InSharedStruct.GetMemory() == SharedInstancedActorDataStruct.GetMemory()
+		, TEXT("We don't expect to override this value"));
 
-	SharedInstancedActorDataStruct = InSharedStruct;
+	SharedInstancedActorDataStruct = TStructView<FInstancedActorsDataSharedFragment>(InSharedStruct.GetMemory());
 }
 
 FString UInstancedActorsData::GetDebugName(bool bCompact) const
diff --git a/Engine/Plugins/Runtime/InstancedActors/Source/InstancedActors/Private/InstancedActorsStationaryLODBatchProcessor.cpp b/Engine/Plugins/Runtime/InstancedActors/Source/InstancedActors/Private/InstancedActorsStationaryLODBatchProcessor.cpp
--- a/Engine/Plugins/Runtime/InstancedActors/Source/InstancedActors/Private/InstancedActorsStationaryLODBatchProcessor.cpp
+++ b/Engine/Plugins/Runtime/InstancedActors/Source/InstancedActors/Private/InstancedActorsStationaryLODBatchProcessor.cpp
@@ -383,37 +383,16 @@ void UInstancedActorsStationaryLODBatchProcessor::Execute(FMassEntityManager& En
 				return NextTickTime;
 			};
 
-		TConstArrayView<FSharedStruct> AllSharedFragmentsOfType = EntityManager.GetSharedFragmentsOfType<FInstancedActorsDataSharedFragment>();
-		if (AllSharedFragmentsOfType.Num() > 0)
+		TArray<UInstancedActorsSubsystem::FNextTickSharedFragment>& SortedSharedFragments = InstancedActorSubsystem->GetTickableSharedFragments();
+		if (SortedSharedFragments.Num() > 0)
 		{
-			if (SortedSharedFragments.Num() == 0)
-			{
-				SortedSharedFragments.Reserve(AllSharedFragmentsOfType.Num());
-				for (const FSharedStruct& SharedStruct : AllSharedFragmentsOfType)
-				{
-					SortedSharedFragments.Add({ SharedStruct });
-				}
-				// we should call SortedSharedFragments.Heapify() but there's no point since all elements have the same NextTickTime now (0).
-			}
-			else if (SortedSharedFragments.Num() < AllSharedFragmentsOfType.Num())
-			{
-				// We add all of them at the front for immediate processing.
-				const int32 StartingIndex = SortedSharedFragments.Num();
-				const int32 NewItemsCount = (AllSharedFragmentsOfType.Num() - SortedSharedFragments.Num());
-				SortedSharedFragments.InsertDefaulted(0, NewItemsCount);
-				for (int32 NewIndex = 0; NewIndex < NewItemsCount; ++NewIndex)
-				{
-					SortedSharedFragments[NewIndex].SharedStruct = AllSharedFragmentsOfType[StartingIndex + NewIndex];
-				}
-				SortedSharedFragments.Heapify();
-			}
-
 			while (SortedSharedFragments.HeapTop().NextTickTime < CurrentTime)
 			{
-				FNextTickSharedFragment WrappedSharedFragment;
+				UInstancedActorsSubsystem::FNextTickSharedFragment WrappedSharedFragment;
 				SortedSharedFragments.HeapPop(WrappedSharedFragment, EAllowShrinking::No);
 				FInstancedActorsDataSharedFragment& ManagerSharedFragment = WrappedSharedFragment.SharedStruct.Get<FInstancedActorsDataSharedFragment>();
 				
+				ManagerSharedFragment.LastTickTime = CurrentTime;
 				WrappedSharedFragment.NextTickTime = ExecutionFunction(ManagerSharedFragment);
 				SortedSharedFragments.HeapPush(MoveTemp(WrappedSharedFragment));
 			}
diff --git a/Engine/Plugins/Runtime/InstancedActors/Source/InstancedActors/Private/InstancedActorsSubsystem.cpp b/Engine/Plugins/Runtime/InstancedActors/Source/InstancedActors/Private/InstancedActorsSubsystem.cpp
--- a/Engine/Plugins/Runtime/InstancedActors/Source/InstancedActors/Private/InstancedActorsSubsystem.cpp
+++ b/Engine/Plugins/Runtime/InstancedActors/Source/InstancedActors/Private/InstancedActorsSubsystem.cpp
@@ -16,6 +16,7 @@
 #include "Engine/Level.h"
 #include "EngineUtils.h"
 #include "MassEntityTypes.h"
+#include "MassEntitySubsystem.h"
 #include "Misc/ArchiveMD5.h"
 #include "Misc/ReverseIterate.h"
 
@@ -152,6 +153,10 @@ void UInstancedActorsSubsystem::Initialize(FSubsystemCollectionBase& Collection)
 	UWorld* World = GetWorld();
 	check(World);
 
+	UMassEntitySubsystem* EntitySubsystem = Collection.InitializeDependency<UMassEntitySubsystem>();
+	check(EntitySubsystem);
+	EntityManager = EntitySubsystem->GetMutableEntityManager().AsShared();
+
 	// As playlist GFP's are initialized after main map load, we account for latent subsystem creation here by registering any existing
 	// AInstancedActorsModifierVolume's and AInstancedActorsManager's that may already have loaded before subsystem creation.
 
@@ -180,7 +185,6 @@ void UInstancedActorsSubsystem::Initialize(FSubsystemCollectionBase& Collection)
 
 void UInstancedActorsSubsystem::Deinitialize()
 {
-
 #if WITH_EDITOR
 	InstancedActorsCVars::CVarRefreshSettings.AsVariable()->SetOnChangedCallback(FConsoleVariableDelegate());
 #endif
@@ -189,6 +193,7 @@ void UInstancedActorsSubsystem::Deinitialize()
 
 	TRACE_CPUPROFILER_EVENT_SCOPE(UInstancedActorsSubsystem Deinitialize);
 
+	EntityManager.Reset();
 	ExemplarActors.Reset();
 
 	if (IsValid(ExemplarActorWorld))
@@ -820,3 +825,61 @@ FInstancedActorsVisualizationDesc UInstancedActorsSubsystem::CreateVisualDescrip
 {
 	return FInstancedActorsVisualizationDesc::FromActor(ExemplarActor);
 }
+
+TArray<UInstancedActorsSubsystem::FNextTickSharedFragment>& UInstancedActorsSubsystem::GetTickableSharedFragments()
+{	
+	RegisterNewSharedFragmentsInternal();
+	return SortedSharedFragments;
+}
+
+void UInstancedActorsSubsystem::UpdateAndResetTickTime(TConstStructView<FInstancedActorsDataSharedFragment> InstancedActorsDataSharedFragment)
+{
+	TRACE_CPUPROFILER_EVENT_SCOPE(UInstancedActorsSubsystem_UpdateTickableSharedFragments);
+
+	const bool bParamStructFound = RegisterNewSharedFragmentsInternal(InstancedActorsDataSharedFragment);
+	if (bParamStructFound == false)
+	{
+		// we need to find the one. Naive implementation for now.
+		// starting from back in assumption stuff has been removed and re-added so the relevant data should be closer 
+		// to the back due to being at EMassLOD::Off level
+		for (int32 Index = SortedSharedFragments.Num() - 1; Index >= 0; --Index)
+		{
+			if (SortedSharedFragments[Index].SharedStruct == InstancedActorsDataSharedFragment)
+			{
+				FNextTickSharedFragment TickFragment = SortedSharedFragments[Index];
+				// setting to 0 will force update the very next time Batch LOD is being calculated. 
+				TickFragment.NextTickTime = 0;
+				SortedSharedFragments.HeapRemoveAt(Index, EAllowShrinking::No);
+				SortedSharedFragments.HeapPush(MoveTemp(TickFragment));
+				break;
+			}
+		}
+	}
+}
+
+bool UInstancedActorsSubsystem::RegisterNewSharedFragmentsInternal(TConstStructView<FInstancedActorsDataSharedFragment> InstancedActorsDataSharedFragment)
+{
+	check(EntityManager);
+
+	const bool bParamStructProvided = InstancedActorsDataSharedFragment.IsValid();
+	// starting with !bParamStructProvided will result in short-circuiting the assignment operation below if InstancedActorsDataSharedFragment is empty
+	bool bParamStructFound = !bParamStructProvided;
+	TConstArrayView<FSharedStruct> AllSharedFragmentsOfType = EntityManager->GetSharedFragmentsOfType<FInstancedActorsDataSharedFragment>();
+	
+	if (SortedSharedFragments.Num() < AllSharedFragmentsOfType.Num())
+	{
+		// We add all of them at the front for immediate processing.
+		const int32 StartingIndex = SortedSharedFragments.Num();
+		const int32 NewItemsCount = (AllSharedFragmentsOfType.Num() - SortedSharedFragments.Num());
+		SortedSharedFragments.InsertDefaulted(0, NewItemsCount);
+		for (int32 NewIndex = 0; NewIndex < NewItemsCount; ++NewIndex)
+		{
+			SortedSharedFragments[NewIndex].SharedStruct = AllSharedFragmentsOfType[StartingIndex + NewIndex];
+			bParamStructFound = bParamStructFound || (AllSharedFragmentsOfType[StartingIndex + NewIndex] == InstancedActorsDataSharedFragment);
+		}
+		SortedSharedFragments.Heapify();
+	}
+	check(SortedSharedFragments.Num() == AllSharedFragmentsOfType.Num());
+
+	return bParamStructFound;
+}
\ No newline at end of file
diff --git a/Engine/Plugins/Runtime/InstancedActors/Source/InstancedActors/Public/InstancedActorsData.h b/Engine/Plugins/Runtime/InstancedActors/Source/InstancedActors/Public/InstancedActorsData.h
--- a/Engine/Plugins/Runtime/InstancedActors/Source/InstancedActors/Public/InstancedActorsData.h
+++ b/Engine/Plugins/Runtime/InstancedActors/Source/InstancedActors/Public/InstancedActorsData.h
@@ -434,7 +434,7 @@ protected:
 
 private:
 	// Represents the shared fragment registered with MassEntityManager, that points back to this UInstancedActorsData instance
-	FSharedStruct SharedInstancedActorDataStruct;
+	TStructView<FInstancedActorsDataSharedFragment> SharedInstancedActorDataStruct;
 
 	// List of FInstancedActorsDelta's to apply to instances, replicated via fast array replication to clients.
 	// Formed at runtime via player actions and persisted / restored by game's persistence system
diff --git a/Engine/Plugins/Runtime/InstancedActors/Source/InstancedActors/Public/InstancedActorsStationaryLODBatchProcessor.h b/Engine/Plugins/Runtime/InstancedActors/Source/InstancedActors/Public/InstancedActorsStationaryLODBatchProcessor.h
--- a/Engine/Plugins/Runtime/InstancedActors/Source/InstancedActors/Public/InstancedActorsStationaryLODBatchProcessor.h
+++ b/Engine/Plugins/Runtime/InstancedActors/Source/InstancedActors/Public/InstancedActorsStationaryLODBatchProcessor.h
@@ -26,18 +26,4 @@ protected:
 
 	UPROPERTY(EditDefaultsOnly, Category="Mass", config)
 	double DelayPerBulkLOD[(int)EInstancedActorsBulkLOD::MAX];
-
-	struct FNextTickSharedFragment
-	{
-		FSharedStruct SharedStruct;
-		double NextTickTime = 0;
-
-		bool operator<(const FNextTickSharedFragment& Other) const
-		{
-			return NextTickTime < Other.NextTickTime;
-		}
-	};
-
-	/** The container storing a sorted queue of FSharedStruct instances, ordered by the NextTickTime */
-	TArray<FNextTickSharedFragment> SortedSharedFragments;
 };
diff --git a/Engine/Plugins/Runtime/InstancedActors/Source/InstancedActors/Public/InstancedActorsSubsystem.h b/Engine/Plugins/Runtime/InstancedActors/Source/InstancedActors/Public/InstancedActorsSubsystem.h
--- a/Engine/Plugins/Runtime/InstancedActors/Source/InstancedActors/Public/InstancedActorsSubsystem.h
+++ b/Engine/Plugins/Runtime/InstancedActors/Source/InstancedActors/Public/InstancedActorsSubsystem.h
@@ -158,8 +158,34 @@ public:
 	INSTANCEDACTORS_API virtual void Initialize(FSubsystemCollectionBase& Collection) override;
 	INSTANCEDACTORS_API virtual void Deinitialize() override;
 	//~ End USubsystem Overrides
-		
+
+	struct FNextTickSharedFragment
+	{
+		FSharedStruct SharedStruct;
+		double NextTickTime = 0;
+
+		bool operator<(const FNextTickSharedFragment& Other) const
+		{
+			return NextTickTime < Other.NextTickTime;
+		}
+	};
+
+	TArray<UInstancedActorsSubsystem::FNextTickSharedFragment>& GetTickableSharedFragments();
+	void UpdateAndResetTickTime(TConstStructView<FInstancedActorsDataSharedFragment> InstancedActorsDataSharedFragment);
+
 protected:
+	/** 
+	 * Fetches all registered FInstancedActorsDataSharedFragment from the EntityManager and adds the missing ones to SortedSharedFragments
+	 * @param InstancedActorsDataSharedFragment optionally the function can check if given shared fragment is amongst 
+	 *	the newly added fragments
+	 * @param returns whether InstancedActorsDataSharedFragment has been found, or `true` if that param is not provided. 
+	 */
+	bool RegisterNewSharedFragmentsInternal(TConstStructView<FInstancedActorsDataSharedFragment> InstancedActorsDataSharedFragment = TConstStructView<FInstancedActorsDataSharedFragment>());
+
+	/** The container storing a sorted queue of FSharedStruct instances, ordered by the NextTickTime */
+	TArray<FNextTickSharedFragment> SortedSharedFragments;
+
+	TSharedPtr<FMassEntityManager> EntityManager;
 
 	UPROPERTY(Transient)
 	TObjectPtr<const UInstancedActorsProjectSettings> ProjectSettings;
diff --git a/Engine/Plugins/Runtime/InstancedActors/Source/InstancedActors/Public/InstancedActorsTypes.h b/Engine/Plugins/Runtime/InstancedActors/Source/InstancedActors/Public/InstancedActorsTypes.h
--- a/Engine/Plugins/Runtime/InstancedActors/Source/InstancedActors/Public/InstancedActorsTypes.h
+++ b/Engine/Plugins/Runtime/InstancedActors/Source/InstancedActors/Public/InstancedActorsTypes.h
@@ -292,6 +292,8 @@ struct FInstancedActorsDataSharedFragment : public FMassSharedFragment
 	TWeakObjectPtr<UInstancedActorsData> InstanceData;
 
 	EInstancedActorsBulkLOD BulkLOD = EInstancedActorsBulkLOD::MAX;
+
+	double LastTickTime = 0.0;
 };
 
 USTRUCT()
