Proper fix for InstancedActors handling of failed instant-hydration.
In addition the change also fixes the FMassActorSpawnRequest leakage - instant-hydration spawn requests were never being removed.

diff --git a/Engine/Plugins/Runtime/InstancedActors/Source/InstancedActors/Private/InstancedActorsRepresentationActorManagement.cpp b/Engine/Plugins/Runtime/InstancedActors/Source/InstancedActors/Private/InstancedActorsRepresentationActorManagement.cpp
--- a/Engine/Plugins/Runtime/InstancedActors/Source/InstancedActors/Private/InstancedActorsRepresentationActorManagement.cpp
+++ b/Engine/Plugins/Runtime/InstancedActors/Source/InstancedActors/Private/InstancedActorsRepresentationActorManagement.cpp
@@ -163,10 +163,18 @@ AActor* UInstancedActorsRepresentationActorManagement::FindOrInstantlySpawnActor
 		// force spawn
 		UMassActorSpawnerSubsystem* ActorSpawnerSubsystem = RepresentationSubsystem.GetActorSpawnerSubsystem();
 		check(ActorSpawnerSubsystem);
-		if (ActorSpawnerSubsystem->ProcessSpawnRequest(Representation.ActorSpawnRequestHandle) != ESpawnRequestStatus::None)
+		const ESpawnRequestStatus RequestResult = ActorSpawnerSubsystem->ProcessSpawnRequest(Representation.ActorSpawnRequestHandle);
+		
+		if (RequestResult == ESpawnRequestStatus::Failed || RequestResult == ESpawnRequestStatus::Succeeded)
 		{
-			FMassActorSpawnRequest& SpawnRequest = ActorSpawnerSubsystem->GetMutableSpawnRequest<FMassActorSpawnRequest>(Representation.ActorSpawnRequestHandle);
-			SpawnedActor = SpawnRequest.SpawnedActor;
+			if (RequestResult == ESpawnRequestStatus::Succeeded)
+			{
+				FMassActorSpawnRequest& SpawnRequest = ActorSpawnerSubsystem->GetMutableSpawnRequest<FMassActorSpawnRequest>(Representation.ActorSpawnRequestHandle);
+				SpawnedActor = SpawnRequest.SpawnedActor;
+				ensureMsgf(ActorSpawnerSubsystem->RemoveActorSpawnRequest(Representation.ActorSpawnRequestHandle), TEXT("Unable to remove a valid spawn request"));
+			}
+			
+			Representation.ActorSpawnRequestHandle.Invalidate();
 		}
 	}
 
