Constified UWorld acces in UE::InstancedActors::Utils functions

diff --git a/Engine/Plugins/Runtime/InstancedActors/Source/InstancedActors/Private/InstancedActorsTypes.cpp b/Engine/Plugins/Runtime/InstancedActors/Source/InstancedActors/Private/InstancedActorsTypes.cpp
--- a/Engine/Plugins/Runtime/InstancedActors/Source/InstancedActors/Private/InstancedActorsTypes.cpp
+++ b/Engine/Plugins/Runtime/InstancedActors/Source/InstancedActors/Private/InstancedActorsTypes.cpp
@@ -12,7 +12,7 @@ DEFINE_LOG_CATEGORY(LogInstancedActors);
 
 namespace UE::InstancedActors::Utils
 {
-	TSubclassOf<UMassActorSpawnerSubsystem> DetermineActorSpawnerSubsystemClass(UWorld& World)
+	TSubclassOf<UMassActorSpawnerSubsystem> DetermineActorSpawnerSubsystemClass(const UWorld& World)
 	{
 		// @todo Add support for non-replay NM_Standalone where we should use UServerInstancedActorsSpawnerSubsystem for 
 		// authoritative actor spawning.
@@ -24,21 +24,21 @@ namespace UE::InstancedActors::Utils
 		return GET_INSTANCEDACTORS_CONFIG_VALUE(GetClientActorSpawnerSubsystemClass());
 	}
 
-	UServerInstancedActorsSpawnerSubsystem* GetServerInstancedActorsSpawnerSubsystem(UWorld& World)
+	UServerInstancedActorsSpawnerSubsystem* GetServerInstancedActorsSpawnerSubsystem(const UWorld& World)
 	{
 		TSubclassOf<UMassActorSpawnerSubsystem> SpawnerSubsystemClass = GET_INSTANCEDACTORS_CONFIG_VALUE(GetServerActorSpawnerSubsystemClass());
 		check(SpawnerSubsystemClass);
 		return Cast<UServerInstancedActorsSpawnerSubsystem>(World.GetSubsystemBase(SpawnerSubsystemClass));
 	}
 
-	UClientInstancedActorsSpawnerSubsystem* GetClientInstancedActorsSpawnerSubsystem(UWorld& World)
+	UClientInstancedActorsSpawnerSubsystem* GetClientInstancedActorsSpawnerSubsystem(const UWorld& World)
 	{
 		TSubclassOf<UMassActorSpawnerSubsystem> SpawnerSubsystemClass = GET_INSTANCEDACTORS_CONFIG_VALUE(GetClientActorSpawnerSubsystemClass());
 		check(SpawnerSubsystemClass);
 		return Cast<UClientInstancedActorsSpawnerSubsystem>(World.GetSubsystemBase(SpawnerSubsystemClass));
 	}
 
-	UMassActorSpawnerSubsystem* GetActorSpawnerSubsystem(UWorld& World)
+	UMassActorSpawnerSubsystem* GetActorSpawnerSubsystem(const UWorld& World)
 	{
 		if (World.GetNetMode() == NM_DedicatedServer)
 		{
@@ -48,7 +48,7 @@ namespace UE::InstancedActors::Utils
 		return GetClientInstancedActorsSpawnerSubsystem(World);
 	}
 
-	UInstancedActorsSubsystem* GetInstancedActorsSubsystem(UWorld& World)
+	UInstancedActorsSubsystem* GetInstancedActorsSubsystem(const UWorld& World)
 	{
 		TSubclassOf<UInstancedActorsSubsystem> InstancedActorsSubsystemClass = GET_INSTANCEDACTORS_CONFIG_VALUE(GetInstancedActorsSubsystemClass());
 		check(InstancedActorsSubsystemClass);
diff --git a/Engine/Plugins/Runtime/InstancedActors/Source/InstancedActors/Public/InstancedActorsTypes.h b/Engine/Plugins/Runtime/InstancedActors/Source/InstancedActors/Public/InstancedActorsTypes.h
--- a/Engine/Plugins/Runtime/InstancedActors/Source/InstancedActors/Public/InstancedActorsTypes.h
+++ b/Engine/Plugins/Runtime/InstancedActors/Source/InstancedActors/Public/InstancedActorsTypes.h
@@ -44,15 +44,15 @@ ENUM_CLASS_FLAGS(EInstancedActorsFragmentFlags);
 
 namespace UE::InstancedActors::Utils
 {
-	INSTANCEDACTORS_API TSubclassOf<UMassActorSpawnerSubsystem> DetermineActorSpawnerSubsystemClass(UWorld& World);
-	INSTANCEDACTORS_API UServerInstancedActorsSpawnerSubsystem* GetServerInstancedActorsSpawnerSubsystem(UWorld& World);
-	INSTANCEDACTORS_API UClientInstancedActorsSpawnerSubsystem* GetClientInstancedActorsSpawnerSubsystem(UWorld& World);
+	INSTANCEDACTORS_API TSubclassOf<UMassActorSpawnerSubsystem> DetermineActorSpawnerSubsystemClass(const UWorld& World);
+	INSTANCEDACTORS_API UServerInstancedActorsSpawnerSubsystem* GetServerInstancedActorsSpawnerSubsystem(const UWorld& World);
+	INSTANCEDACTORS_API UClientInstancedActorsSpawnerSubsystem* GetClientInstancedActorsSpawnerSubsystem(const UWorld& World);
 	/** 
 	 * Calls either GetServerInstancedActorsSpawnerSubsystem or GetClientInstancedActorsSpawnerSubsystem, depending on 
 	 * given UWorld's net mode.
 	 */
-	INSTANCEDACTORS_API UMassActorSpawnerSubsystem* GetActorSpawnerSubsystem(UWorld& World);
-	INSTANCEDACTORS_API UInstancedActorsSubsystem* GetInstancedActorsSubsystem(UWorld& World);
+	INSTANCEDACTORS_API UMassActorSpawnerSubsystem* GetActorSpawnerSubsystem(const UWorld& World);
+	INSTANCEDACTORS_API UInstancedActorsSubsystem* GetInstancedActorsSubsystem(const UWorld& World);
 }
 
 // FInstancedActorsTagSet -> FInstancedActorsTagSet
