Added means of querying the InstancedActorsSubsystem for instances of a given actor class inside a given bounds.

diff --git a/Engine/Plugins/Runtime/InstancedActors/Source/InstancedActors/Private/InstancedActorsManager.cpp b/Engine/Plugins/Runtime/InstancedActors/Source/InstancedActors/Private/InstancedActorsManager.cpp
--- a/Engine/Plugins/Runtime/InstancedActors/Source/InstancedActors/Private/InstancedActorsManager.cpp
+++ b/Engine/Plugins/Runtime/InstancedActors/Source/InstancedActors/Private/InstancedActorsManager.cpp
@@ -793,7 +793,7 @@ bool AInstancedActorsManager::ForEachInstance(FInstanceOperationFunc Operation,
 			if (InstancedActorDataPredicate.IsSet())
 			{
 				const bool bPassedPredicate = ::Invoke(*InstancedActorDataPredicate, *InstanceData);
-				if (UNLIKELY(!bPassedPredicate))
+				if (!bPassedPredicate)
 				{
 					continue;
 				}
@@ -954,6 +954,77 @@ bool AInstancedActorsManager::ForEachInstance<FBox>(const FBox& QueryBounds, FIn
 template bool AInstancedActorsManager::ForEachInstance<FBox>(const FBox& QueryBounds, AInstancedActorsManager::FInstanceOperationFunc Operation) const;
 template bool AInstancedActorsManager::ForEachInstance<FSphere>(const FSphere& QueryBounds, AInstancedActorsManager::FInstanceOperationFunc Operation) const;
 
+template<>
+bool AInstancedActorsManager::IsInstanceInsideBounds<FBox>(const FBox& QueryBounds, const FInstancedActorsInstanceHandle& InstanceHandle, const FTransform& InstanceTransform)
+{
+	if (QueryBounds.IsInside(InstanceTransform.GetLocation()))
+	{
+		return true;
+	}
+
+	// More expensive bounds test.
+	const FBox InstancedActorBounds = CalculateBounds(InstanceHandle.InstancedActorData->ActorClass).TransformBy(InstanceTransform);
+	if (QueryBounds.Intersect(InstancedActorBounds))
+	{
+		return true;
+	}
+
+	return false;
+}
+
+template<>
+bool AInstancedActorsManager::IsInstanceInsideBounds<FSphere>(const FSphere& QueryBounds, const FInstancedActorsInstanceHandle& InstanceHandle, const FTransform& InstanceTransform)
+{
+	if (QueryBounds.IsInside(InstanceTransform.GetLocation()))
+	{
+		return true;
+	}
+
+	// More expensive bounds test.
+	const FBox InstancedActorBounds = CalculateBounds(InstanceHandle.InstancedActorData->ActorClass);
+	const FSphere TransformedSphere = QueryBounds.TransformBy(InstanceTransform.Inverse());
+	if (FMath::SphereAABBIntersection(TransformedSphere, InstancedActorBounds))
+	{
+		return true;
+	}
+
+	return false;
+}
+
+template<typename TBoundsType>
+bool AInstancedActorsManager::HasInstancesOfClass(const TBoundsType& QueryBounds, TSubclassOf<AActor> ActorClass) const
+{
+	ensure(ActorClass);
+
+	FScopedInstancedActorsIterationContext IterationContext;
+	bool bHasInstance = false;
+
+	auto InstancedActorDataMask = TOptional<AInstancedActorsManager::FInstancedActorDataPredicateFunc>([ActorClass](const UInstancedActorsData& InstancedActorData)
+		{
+			return InstancedActorData.ActorClass->IsChildOf(ActorClass);
+		});
+
+	ForEachInstance([Manager = this, QueryBounds, ActorClass, &bHasInstance](const FInstancedActorsInstanceHandle& InstanceHandle, const FTransform& InstanceTransform, FInstancedActorsIterationContext& IterationContext)
+		{
+			bHasInstance = Manager->IsInstanceInsideBounds(QueryBounds, InstanceHandle, InstanceTransform);
+			UE_IFVLOG(
+			if (bHasInstance)
+			{
+				const FBox InstancedActorBounds = CalculateBounds(InstanceHandle.InstancedActorData->ActorClass).TransformBy(InstanceTransform);
+				UE_VLOG_BOX(Manager, LogInstancedActors, Log, InstancedActorBounds, FColor::Red, TEXT("Instance of class %s"), *GetNameSafe(ActorClass));
+			})
+			const bool bContinue = !bHasInstance;
+			return bContinue;
+		}
+		, IterationContext, InstancedActorDataMask);
+
+	return bHasInstance;
+}
+
+// Instantiate FBox and FSphere implementations
+template bool AInstancedActorsManager::HasInstancesOfClass<FBox>(const FBox& QueryBounds, TSubclassOf<AActor> ActorClass) const;
+template bool AInstancedActorsManager::HasInstancesOfClass<FSphere>(const FSphere& QueryBounds, TSubclassOf<AActor> ActorClass) const;
+
 void AInstancedActorsManager::AuditInstances(FOutputDevice& Ar, bool bDebugDraw, float DebugDrawDuration) const
 {
 	UWorld* World = GetWorld();
diff --git a/Engine/Plugins/Runtime/InstancedActors/Source/InstancedActors/Private/InstancedActorsSubsystem.cpp b/Engine/Plugins/Runtime/InstancedActors/Source/InstancedActors/Private/InstancedActorsSubsystem.cpp
--- a/Engine/Plugins/Runtime/InstancedActors/Source/InstancedActors/Private/InstancedActorsSubsystem.cpp
+++ b/Engine/Plugins/Runtime/InstancedActors/Source/InstancedActors/Private/InstancedActorsSubsystem.cpp
@@ -575,6 +575,19 @@ void UInstancedActorsSubsystem::ForEachInstance(const FBox& QueryBounds, TFuncti
 	});
 }
 
+bool UInstancedActorsSubsystem::HasInstancesOfClass(const FBox& QueryBounds, TSubclassOf<AActor> ActorClass) const
+{
+	bool bHasInstances = false;
+	ForEachManager(QueryBounds, [QueryBounds, ActorClass, &bHasInstances](AInstancedActorsManager& Manager)
+	{
+		bHasInstances = Manager.HasInstancesOfClass(QueryBounds, ActorClass);
+		const bool bContinue = !bHasInstances;
+		return bContinue;
+	});
+
+	return bHasInstances;
+}
+
 AActor& UInstancedActorsSubsystem::GetOrCreateExemplarActor(TSubclassOf<AActor> ActorClass)
 {
 	TRACE_CPUPROFILER_EVENT_SCOPE(UInstancedActorsSubsystem GetOrCreateExemplarActor);
diff --git a/Engine/Plugins/Runtime/InstancedActors/Source/InstancedActors/Public/InstancedActorsManager.h b/Engine/Plugins/Runtime/InstancedActors/Source/InstancedActors/Public/InstancedActorsManager.h
--- a/Engine/Plugins/Runtime/InstancedActors/Source/InstancedActors/Public/InstancedActorsManager.h
+++ b/Engine/Plugins/Runtime/InstancedActors/Source/InstancedActors/Public/InstancedActorsManager.h
@@ -150,6 +150,21 @@ public:
 	template <typename TBoundsType>
 	bool ForEachInstance(const TBoundsType& QueryBounds, FInstanceOperationFunc InOperation, FInstancedActorsIterationContext& IterationContext, TOptional<FInstancedActorDataPredicateFunc> InstancedActorDataPredicate = TOptional<FInstancedActorDataPredicateFunc>()) const;
 
+	/**
+	 * Checks whether there are any instanced actors within this manager, representing ActorClass or its subclasses inside QueryBounds.
+	 * The check doesn't differentiate between hydrated and dehydrated actors (i.e. whether there's an actor instance
+	 * associated with the instance or not).
+	 */
+	template <typename TBoundsType>
+	bool HasInstancesOfClass(const TBoundsType& QueryBounds, TSubclassOf<AActor> ActorClass) const;
+
+	/** 
+	 * Determines whether the actor instance given by InstanceHandle overlaps QueryBounds. The test involves calculating 
+	 * bounding box of the actor representation of the given instance (i.e. it's not only the transform that's being tested).
+	 */
+	template <typename TBoundsType>
+	static bool IsInstanceInsideBounds(const TBoundsType& QueryBounds, const FInstancedActorsInstanceHandle& InstanceHandle, const FTransform& InstanceTransform);
+
 	// Outputs instance metrics to Ar
 	void AuditInstances(FOutputDevice& Ar, bool bDebugDraw = false, float DebugDrawDuration = 10.0f) const;
 
diff --git a/Engine/Plugins/Runtime/InstancedActors/Source/InstancedActors/Public/InstancedActorsSubsystem.h b/Engine/Plugins/Runtime/InstancedActors/Source/InstancedActors/Public/InstancedActorsSubsystem.h
--- a/Engine/Plugins/Runtime/InstancedActors/Source/InstancedActors/Public/InstancedActorsSubsystem.h
+++ b/Engine/Plugins/Runtime/InstancedActors/Source/InstancedActors/Public/InstancedActorsSubsystem.h
@@ -72,6 +72,13 @@ public:
 	INSTANCEDACTORS_API void ForEachModifierVolume(const FBox& QueryBounds, TFunctionRef<bool(UInstancedActorsModifierVolumeComponent&)> InOperation) const;
 	INSTANCEDACTORS_API void ForEachInstance(const FBox& QueryBounds, TFunctionRef<bool(const FInstancedActorsInstanceHandle&, const FTransform&, FInstancedActorsIterationContext&)> InOperation) const;
 
+	/** 
+	 * Checks whether there are any instanced actors representing ActorClass or its subclasses inside QueryBounds.
+	 * The check doesn't differentiate between hydrated and dehydrated actors (i.e. whether there's an actor instance 
+	 * associated with the instance or not).
+	 */
+	INSTANCEDACTORS_API bool HasInstancesOfClass(const FBox& QueryBounds, TSubclassOf<AActor> ActorClass) const;
+
 	FInstancedActorsManagerHandle AddManager(AInstancedActorsManager& Manager);
 	void RemoveManager(FInstancedActorsManagerHandle ManagerHandle);
 
