Modified InstancedActors' HasInstancesOfClass to optionally precision-test spawned actors when available.

diff --git a/Engine/Plugins/Runtime/InstancedActors/Source/InstancedActors/Private/InstancedActorsManager.cpp b/Engine/Plugins/Runtime/InstancedActors/Source/InstancedActors/Private/InstancedActorsManager.cpp
--- a/Engine/Plugins/Runtime/InstancedActors/Source/InstancedActors/Private/InstancedActorsManager.cpp
+++ b/Engine/Plugins/Runtime/InstancedActors/Source/InstancedActors/Private/InstancedActorsManager.cpp
@@ -828,7 +828,8 @@ bool AInstancedActorsManager::ForEachInstance(FInstanceOperationFunc Operation,
 						{
 							return;
 						}
-					} });
+					}
+				});
 
 			if (!bContinue)
 			{
@@ -955,29 +956,30 @@ template bool AInstancedActorsManager::ForEachInstance<FBox>(const FBox& QueryBo
 template bool AInstancedActorsManager::ForEachInstance<FSphere>(const FSphere& QueryBounds, AInstancedActorsManager::FInstanceOperationFunc Operation) const;
 
 template<>
-bool AInstancedActorsManager::IsInstanceInsideBounds<FBox>(const FBox& QueryBounds, const FInstancedActorsInstanceHandle& InstanceHandle, const FTransform& InstanceTransform)
+UE::InstancedActors::EInsideBoundsTestResult AInstancedActorsManager::IsInstanceInsideBounds<FBox>(const FBox& QueryBounds
+	, const FInstancedActorsInstanceHandle& InstanceHandle, const FTransform& InstanceTransform)
 {
 	if (QueryBounds.IsInside(InstanceTransform.GetLocation()))
 	{
-		return true;
+		return UE::InstancedActors::EInsideBoundsTestResult::OverlapLocation;
 	}
 
 	// More expensive bounds test.
 	const FBox InstancedActorBounds = CalculateBounds(InstanceHandle.InstancedActorData->ActorClass).TransformBy(InstanceTransform);
 	if (QueryBounds.Intersect(InstancedActorBounds))
 	{
-		return true;
+		return UE::InstancedActors::EInsideBoundsTestResult::OverlapBounds;
 	}
 
-	return false;
+	return UE::InstancedActors::EInsideBoundsTestResult::NotInside;
 }
 
 template<>
-bool AInstancedActorsManager::IsInstanceInsideBounds<FSphere>(const FSphere& QueryBounds, const FInstancedActorsInstanceHandle& InstanceHandle, const FTransform& InstanceTransform)
+UE::InstancedActors::EInsideBoundsTestResult AInstancedActorsManager::IsInstanceInsideBounds<FSphere>(const FSphere& QueryBounds, const FInstancedActorsInstanceHandle& InstanceHandle, const FTransform& InstanceTransform)
 {
 	if (QueryBounds.IsInside(InstanceTransform.GetLocation()))
 	{
-		return true;
+		return UE::InstancedActors::EInsideBoundsTestResult::OverlapLocation;
 	}
 
 	// More expensive bounds test.
@@ -985,16 +987,18 @@ bool AInstancedActorsManager::IsInstanceInsideBounds<FSphere>(const FSphere& Que
 	const FSphere TransformedSphere = QueryBounds.TransformBy(InstanceTransform.Inverse());
 	if (FMath::SphereAABBIntersection(TransformedSphere, InstancedActorBounds))
 	{
-		return true;
+		return UE::InstancedActors::EInsideBoundsTestResult::OverlapBounds;
 	}
 
-	return false;
+	return UE::InstancedActors::EInsideBoundsTestResult::NotInside;
 }
 
-template<typename TBoundsType>
-bool AInstancedActorsManager::HasInstancesOfClass(const TBoundsType& QueryBounds, TSubclassOf<AActor> ActorClass) const
+bool AInstancedActorsManager::HasInstancesOfClass(const FBox& InQueryBounds, TSubclassOf<AActor> ActorClass, const bool bTestActorsIfSpawned) const
 {
+	using UE::InstancedActors::EInsideBoundsTestResult;
+
 	ensure(ActorClass);
+	check(InstancedActorSubsystem);
 
 	FScopedInstancedActorsIterationContext IterationContext;
 	bool bHasInstance = false;
@@ -1004,15 +1008,62 @@ bool AInstancedActorsManager::HasInstancesOfClass(const TBoundsType& QueryBounds
 			return InstancedActorData.ActorClass->IsChildOf(ActorClass);
 		});
 
-	ForEachInstance([Manager = this, QueryBounds, ActorClass, &bHasInstance](const FInstancedActorsInstanceHandle& InstanceHandle, const FTransform& InstanceTransform, FInstancedActorsIterationContext& IterationContext)
+	struct FQueryBounds
+	{
+		FBox QueryBounds;
+		FVector Center;
+		FVector Extent;
+		FCollisionShape CollistionShape;
+
+		FQueryBounds(const FBox& InQueryBounds)
+			: QueryBounds(InQueryBounds)
 		{
-			bHasInstance = Manager->IsInstanceInsideBounds(QueryBounds, InstanceHandle, InstanceTransform);
-			UE_IFVLOG(
-			if (bHasInstance)
+			InQueryBounds.GetCenterAndExtents(Center, Extent);
+			CollistionShape = FCollisionShape::MakeBox(Extent);
+		}
+	};
+	FQueryBounds CachedQueryBounds(InQueryBounds);
+
+	ForEachInstance([Manager = this, CachedQueryBounds, ActorClass, &bHasInstance, InstancedActorSubsystem=InstancedActorSubsystem, bTestActorsIfSpawned](const FInstancedActorsInstanceHandle& InstanceHandle, const FTransform& InstanceTransform, FInstancedActorsIterationContext& IterationContext)
+		{
+			const EInsideBoundsTestResult OverlapResult = Manager->IsInstanceInsideBounds(CachedQueryBounds.QueryBounds, InstanceHandle, InstanceTransform);
+			
+			// the only case we care about, representation bounds overlap without overlapping the actual instance location
+			if (OverlapResult == EInsideBoundsTestResult::OverlapBounds && bTestActorsIfSpawned)
+			{
+				// we want to check if the given entity has an actor representation, and if so then test the actor itself
+				UInstancedActorsData& OwningInstanceData = InstanceHandle.GetInstanceActorDataChecked();
+				if (AActor* Actor = Manager->GetActorForInstance(OwningInstanceData, InstanceHandle.GetIndex()))
+				{
+					bHasInstance = false;
+
+					TArray<UPrimitiveComponent*> OutComponents;
+					Actor->GetComponents(UPrimitiveComponent::StaticClass(), OutComponents, /*bIncludeFromChildActors=*/true);
+					for (UPrimitiveComponent* Component : OutComponents)
+					{
+						if (Component->OverlapComponent(CachedQueryBounds.Center, FQuat::Identity, CachedQueryBounds.CollistionShape))
+						{
+							bHasInstance = true;
+							break;
+						}
+					}
+				}
+			}
+			else 
 			{
-				const FBox InstancedActorBounds = CalculateBounds(InstanceHandle.InstancedActorData->ActorClass).TransformBy(InstanceTransform);
-				UE_VLOG_BOX(Manager, LogInstancedActors, Log, InstancedActorBounds, FColor::Red, TEXT("Instance of class %s"), *GetNameSafe(ActorClass));
-			})
+				bHasInstance = (OverlapResult != EInsideBoundsTestResult::NotInside);
+			}
+				
+			UE_IFVLOG
+			(
+				if (bHasInstance || OverlapResult == EInsideBoundsTestResult::OverlapBounds)
+				{
+					const FBox InstancedActorBounds = CalculateBounds(InstanceHandle.InstancedActorData->ActorClass).TransformBy(InstanceTransform);
+					UE_VLOG_BOX(Manager, LogInstancedActors, Log, InstancedActorBounds, bHasInstance ? FColor::Red : FColor::Green
+						, TEXT("Instance of class %s"), *GetNameSafe(InstanceHandle.InstancedActorData->ActorClass));
+				}
+			);
+
 			const bool bContinue = !bHasInstance;
 			return bContinue;
 		}
@@ -1021,10 +1072,6 @@ bool AInstancedActorsManager::HasInstancesOfClass(const TBoundsType& QueryBounds
 	return bHasInstance;
 }
 
-// Instantiate FBox and FSphere implementations
-template bool AInstancedActorsManager::HasInstancesOfClass<FBox>(const FBox& QueryBounds, TSubclassOf<AActor> ActorClass) const;
-template bool AInstancedActorsManager::HasInstancesOfClass<FSphere>(const FSphere& QueryBounds, TSubclassOf<AActor> ActorClass) const;
-
 void AInstancedActorsManager::AuditInstances(FOutputDevice& Ar, bool bDebugDraw, float DebugDrawDuration) const
 {
 	UWorld* World = GetWorld();
@@ -1502,7 +1549,7 @@ int32 AInstancedActorsManager::ConvertCollisionIndexToInstanceIndex(int32 InInde
 	return INDEX_NONE;
 }
 
-AActor* AInstancedActorsManager::FindActorInternal(const FActorInstanceHandle& Handle, FMassEntityView& OutEntityView, const bool bEnsureOnMissingInstanceDataOrMassEntity)
+AActor* AInstancedActorsManager::FindActorInternal(const FActorInstanceHandle& Handle, FMassEntityView& OutEntityView, const bool bEnsureOnMissingInstanceDataOrMassEntity) const
 {
 #if WITH_EDITOR
 	const UWorld* World = GetWorld();
@@ -1510,7 +1557,7 @@ AActor* AInstancedActorsManager::FindActorInternal(const FActorInstanceHandle& H
 	{
 		// In Editor non-game worlds the manager only creates preview ISM components which
 		// are not fully setup with Mass so we simply return this manager as the instance's actor.
-		return this;
+		return const_cast<AInstancedActorsManager*>(this);
 	}
 #endif // WITH_EDITOR
 
@@ -1528,7 +1575,7 @@ AActor* AInstancedActorsManager::FindActorInternal(const FActorInstanceHandle& H
 	{
 		// On clients, where we don't spawn actors directly, we simply return this manager as the
 		// instance's actor.
-		return this;
+		return const_cast<AInstancedActorsManager*>(this);
 	}
 
 	const int32 CompositeIndex = Handle.GetInstanceIndex();
@@ -1574,16 +1621,31 @@ AActor* AInstancedActorsManager::FindActorInternal(const FActorInstanceHandle& H
 	return nullptr;
 }
 
+AActor* AInstancedActorsManager::GetActorForInstance(const UInstancedActorsData& InstanceData, const int32 InstancedActorIndex) const
+{
+	const FMassEntityHandle EntityHandle = InstanceData.GetEntityHandleForIndex(InstancedActorIndex);
+	// note that it's possible that InstancedActorIndex points at a no-longer-valid entity
+	if (EntityHandle.IsValid())
+	{
+		const FMassEntityManager& EntityManager = GetMassEntityManagerChecked();
+		if (FMassActorFragment* ActorFragment = EntityManager.GetFragmentDataPtr<FMassActorFragment>(EntityHandle))
+		{
+			return ActorFragment->GetMutable();
+		}
+	}
+	return nullptr;
+}
+
 AActor* AInstancedActorsManager::FindActor(const FActorInstanceHandle& Handle)
 {
 	FMassEntityView EntityView;
-	return FindActorInternal(Handle, EntityView, /*bEnsureOnMissingInstanceDataOrMassEntity*/ false);
+	return FindActorInternal(Handle, EntityView, /*bEnsureOnMissingInstanceDataOrMassEntity=*/false);
 }
 
 AActor* AInstancedActorsManager::FindOrCreateActor(const FActorInstanceHandle& Handle)
 {
 	FMassEntityView EntityView;
-	AActor* Actor = FindActorInternal(Handle, EntityView, /*bEnsureOnMissingInstanceDataOrMassEntity*/ true);
+	AActor* Actor = FindActorInternal(Handle, EntityView, /*bEnsureOnMissingInstanceDataOrMassEntity=*/true);
 
 	// Create missing actor from Mass if we have a valid EntityView
 	if (Actor == nullptr && EntityView.IsValid())
diff --git a/Engine/Plugins/Runtime/InstancedActors/Source/InstancedActors/Private/InstancedActorsSubsystem.cpp b/Engine/Plugins/Runtime/InstancedActors/Source/InstancedActors/Private/InstancedActorsSubsystem.cpp
--- a/Engine/Plugins/Runtime/InstancedActors/Source/InstancedActors/Private/InstancedActorsSubsystem.cpp
+++ b/Engine/Plugins/Runtime/InstancedActors/Source/InstancedActors/Private/InstancedActorsSubsystem.cpp
@@ -575,12 +575,14 @@ void UInstancedActorsSubsystem::ForEachInstance(const FBox& QueryBounds, TFuncti
 	});
 }
 
-bool UInstancedActorsSubsystem::HasInstancesOfClass(const FBox& QueryBounds, TSubclassOf<AActor> ActorClass) const
+bool UInstancedActorsSubsystem::HasInstancesOfClass(const FBox& QueryBounds, TSubclassOf<AActor> ActorClass, const bool bTestActorsIfSpawned) const
 {
+	TRACE_CPUPROFILER_EVENT_SCOPE(UInstancedActorsSubsystem::HasInstancesOfClass);
+
 	bool bHasInstances = false;
-	ForEachManager(QueryBounds, [QueryBounds, ActorClass, &bHasInstances](AInstancedActorsManager& Manager)
+	ForEachManager(QueryBounds, [QueryBounds, ActorClass, &bHasInstances, bTestActorsIfSpawned](AInstancedActorsManager& Manager)
 	{
-		bHasInstances = Manager.HasInstancesOfClass(QueryBounds, ActorClass);
+		bHasInstances = Manager.HasInstancesOfClass(QueryBounds, ActorClass, bTestActorsIfSpawned);
 		const bool bContinue = !bHasInstances;
 		return bContinue;
 	});
diff --git a/Engine/Plugins/Runtime/InstancedActors/Source/InstancedActors/Public/InstancedActorsData.h b/Engine/Plugins/Runtime/InstancedActors/Source/InstancedActors/Public/InstancedActorsData.h
--- a/Engine/Plugins/Runtime/InstancedActors/Source/InstancedActors/Public/InstancedActorsData.h
+++ b/Engine/Plugins/Runtime/InstancedActors/Source/InstancedActors/Public/InstancedActorsData.h
@@ -357,6 +357,11 @@ public:
 	int32 GetEntityIndexFromCollisionIndex(const UInstancedStaticMeshComponent& ISMComponent, const int32 CollisionIndex) const;
 
 protected:
+	FMassEntityHandle GetEntityHandleForIndex(const int32 Index) const
+	{
+		return Entities.IsValidIndex(Index) ? Entities[Index] : FMassEntityHandle();
+	}
+
 	//~ Begin UObject Overrides
 	virtual void PostDuplicate(EDuplicateMode::Type DuplicateMode) override;
 	virtual void PostLoad() override;
diff --git a/Engine/Plugins/Runtime/InstancedActors/Source/InstancedActors/Public/InstancedActorsManager.h b/Engine/Plugins/Runtime/InstancedActors/Source/InstancedActors/Public/InstancedActorsManager.h
--- a/Engine/Plugins/Runtime/InstancedActors/Source/InstancedActors/Public/InstancedActorsManager.h
+++ b/Engine/Plugins/Runtime/InstancedActors/Source/InstancedActors/Public/InstancedActorsManager.h
@@ -22,10 +22,20 @@ struct FMassEntityManager;
 class UInstancedActorsData;
 class UInstancedStaticMeshComponent;
 
-namespace UE::InstancedActors::CVars
+namespace UE::InstancedActors
 {
-	extern INSTANCEDACTORS_API bool bEnablePersistence;
-}
+	enum EInsideBoundsTestResult
+	{
+		NotInside,
+		OverlapLocation,
+		OverlapBounds
+	};
+
+	namespace CVars
+	{
+		extern INSTANCEDACTORS_API bool bEnablePersistence;
+	}
+} // UE::InstancedActors
 
 DECLARE_STATS_GROUP(TEXT("InstanceActor Rendering"), STATGROUP_InstancedActorsRendering, STATCAT_Advanced);
 
@@ -148,22 +158,25 @@ public:
 	template <typename TBoundsType>
 	bool ForEachInstance(const TBoundsType& QueryBounds, FInstanceOperationFunc InOperation) const;
 	template <typename TBoundsType>
-	bool ForEachInstance(const TBoundsType& QueryBounds, FInstanceOperationFunc InOperation, FInstancedActorsIterationContext& IterationContext, TOptional<FInstancedActorDataPredicateFunc> InstancedActorDataPredicate = TOptional<FInstancedActorDataPredicateFunc>()) const;
+	bool ForEachInstance(const TBoundsType& QueryBounds, FInstanceOperationFunc InOperation, FInstancedActorsIterationContext& IterationContext
+		, TOptional<FInstancedActorDataPredicateFunc> InstancedActorDataPredicate = TOptional<FInstancedActorDataPredicateFunc>()) const;
 
 	/**
 	 * Checks whether there are any instanced actors within this manager, representing ActorClass or its subclasses inside QueryBounds.
 	 * The check doesn't differentiate between hydrated and dehydrated actors (i.e. whether there's an actor instance
 	 * associated with the instance or not).
+	 * @param bTestActorsIfSpawned if true then when an instance is found to overlap given bounds, and it has an actor 
+	 *	spawned associated with it, then the actor itself will be tested against the bounds for more precise test.
 	 */
-	template <typename TBoundsType>
-	bool HasInstancesOfClass(const TBoundsType& QueryBounds, TSubclassOf<AActor> ActorClass) const;
+	bool HasInstancesOfClass(const FBox& QueryBounds, TSubclassOf<AActor> ActorClass, const bool bTestActorsIfSpawned = false) const;
 
 	/** 
 	 * Determines whether the actor instance given by InstanceHandle overlaps QueryBounds. The test involves calculating 
 	 * bounding box of the actor representation of the given instance (i.e. it's not only the transform that's being tested).
 	 */
 	template <typename TBoundsType>
-	static bool IsInstanceInsideBounds(const TBoundsType& QueryBounds, const FInstancedActorsInstanceHandle& InstanceHandle, const FTransform& InstanceTransform);
+	static UE::InstancedActors::EInsideBoundsTestResult IsInstanceInsideBounds(const TBoundsType & QueryBounds
+		, const FInstancedActorsInstanceHandle& InstanceHandle, const FTransform & InstanceTransform);
 
 	// Outputs instance metrics to Ar
 	void AuditInstances(FOutputDevice& Ar, bool bDebugDraw = false, float DebugDrawDuration = 10.0f) const;
@@ -241,12 +254,12 @@ protected:
 	//~ End AActor Overrides
 
 	//~ Begin IActorInstanceManagerInterface Overrides
-	virtual int32 ConvertCollisionIndexToInstanceIndex(int32 InIndex, const UPrimitiveComponent* RelevantComponent) const /*override*/;
-	virtual AActor* FindActor(const FActorInstanceHandle& Handle) /*override*/;
-	virtual AActor* FindOrCreateActor(const FActorInstanceHandle& Handle) /*override*/;
-	virtual UClass* GetRepresentedClass(const int32 InstanceIndex) const /*override*/;
-	virtual ULevel* GetLevelForInstance(const int32 InstanceIndex) const /*override*/;
-	virtual FTransform GetTransform(const FActorInstanceHandle& Handle) const /*override*/;
+	virtual int32 ConvertCollisionIndexToInstanceIndex(int32 InIndex, const UPrimitiveComponent* RelevantComponent) const override;
+	virtual AActor* FindActor(const FActorInstanceHandle& Handle) override;
+	virtual AActor* FindOrCreateActor(const FActorInstanceHandle& Handle) override;
+	virtual UClass* GetRepresentedClass(const int32 InstanceIndex) const override;
+	virtual ULevel* GetLevelForInstance(const int32 InstanceIndex) const override;
+	virtual FTransform GetTransform(const FActorInstanceHandle& Handle) const override;
 	//~ End IActorInstanceManagerInterface Overrides
 
 	// Called by Serialize for SaveGame archives to save / load IAD persistence data
@@ -362,7 +375,9 @@ private:
 	 * Try to extract the actor from the provided handle or from associated Mass Entity.
 	 * When unable to retrieve it returns nullptr but also the associated EntityView so caller could reuse it to create the actor.
 	 */
-	AActor* FindActorInternal(const FActorInstanceHandle& Handle, FMassEntityView& OutEntityView, bool bEnsureOnMissingInstanceDataOrMassEntity);
+	AActor* FindActorInternal(const FActorInstanceHandle& Handle, FMassEntityView& OutEntityView, bool bEnsureOnMissingInstanceDataOrMassEntity) const;
+
+	AActor* GetActorForInstance(const UInstancedActorsData& InstanceData, const int32 InstancedActorIndex) const;
 
 	FInstancedActorsInstanceHandle ActorInstanceHandleFromFSMInstanceId(const FSMInstanceId& InstanceId) const;
 };
diff --git a/Engine/Plugins/Runtime/InstancedActors/Source/InstancedActors/Public/InstancedActorsSubsystem.h b/Engine/Plugins/Runtime/InstancedActors/Source/InstancedActors/Public/InstancedActorsSubsystem.h
--- a/Engine/Plugins/Runtime/InstancedActors/Source/InstancedActors/Public/InstancedActorsSubsystem.h
+++ b/Engine/Plugins/Runtime/InstancedActors/Source/InstancedActors/Public/InstancedActorsSubsystem.h
@@ -76,8 +76,10 @@ public:
 	 * Checks whether there are any instanced actors representing ActorClass or its subclasses inside QueryBounds.
 	 * The check doesn't differentiate between hydrated and dehydrated actors (i.e. whether there's an actor instance 
 	 * associated with the instance or not).
+	 * * @param bTestActorsIfSpawned if true then when an instance is found to overlap given bounds, and it has an actor
+	 *	spawned associated with it, then the actor itself will be tested against the bounds for more precise test.
 	 */
-	INSTANCEDACTORS_API bool HasInstancesOfClass(const FBox& QueryBounds, TSubclassOf<AActor> ActorClass) const;
+	INSTANCEDACTORS_API bool HasInstancesOfClass(const FBox& QueryBounds, TSubclassOf<AActor> ActorClass, const bool bTestActorsIfSpawned = false) const;
 
 	FInstancedActorsManagerHandle AddManager(AInstancedActorsManager& Manager);
 	void RemoveManager(FInstancedActorsManagerHandle ManagerHandle);
