Set a fixed replication priority for AInstancedActorsManager to ensure updates to itself and UInstancedActorsData are replicated at the maximum configurable frequency when using Iris.

diff --git a/Engine/Plugins/Runtime/InstancedActors/Source/InstancedActors/Private/InstancedActorsManager.cpp b/Engine/Plugins/Runtime/InstancedActors/Source/InstancedActors/Private/InstancedActorsManager.cpp
--- a/Engine/Plugins/Runtime/InstancedActors/Source/InstancedActors/Private/InstancedActorsManager.cpp
+++ b/Engine/Plugins/Runtime/InstancedActors/Source/InstancedActors/Private/InstancedActorsManager.cpp
@@ -28,6 +28,9 @@
 #include "MassRepresentationSubsystem.h"
 #include "Math/NumericLimits.h"
 
+#if UE_WITH_IRIS
+#include "Net/Iris/ReplicationSystem/ReplicationSystemUtil.h"
+#endif
 
 DECLARE_DWORD_ACCUMULATOR_STAT(TEXT("Detailed Instance Count"), STAT_DetailedInstanceCount, STATGROUP_InstancedActorsRendering);
 DECLARE_DWORD_ACCUMULATOR_STAT(TEXT("Medium Instance Count"), STAT_MediumInstanceCount, STATGROUP_InstancedActorsRendering);
@@ -144,6 +147,24 @@ AInstancedActorsManager::AInstancedActorsManager()
 	// @todo need default implementation of GUID and set it here via SetSavedActorGUID
 }
 
+#if UE_WITH_IRIS
+void AInstancedActorsManager::BeginReplication()
+{
+	Super::BeginReplication();
+
+	// This actor is configured to use spatial prioritization in Iris. If it's too far from
+	// a player's position it will replicate to their client less frequently. This becomes a 
+	// problem when destroying instanced actors because the client needs to receive both
+	// the actor destruction notification and instanced actor destruction notification
+	// at the same time. If the aren't received at the same time then the instanced actor's
+	// mesh may still be visible for a period of time after destroying the actor.
+	//
+	// Fixing the priority to 1.0 means that it will always have the maximum priority and
+	// replicated at the maximum configured frequency.
+	UE::Net::FReplicationSystemUtil::SetStaticPriority(this, 1.0f);
+}
+#endif
+
 void AInstancedActorsManager::BeginPlay()
 {
 	Super::BeginPlay();
diff --git a/Engine/Plugins/Runtime/InstancedActors/Source/InstancedActors/Public/InstancedActorsManager.h b/Engine/Plugins/Runtime/InstancedActors/Source/InstancedActors/Public/InstancedActorsManager.h
--- a/Engine/Plugins/Runtime/InstancedActors/Source/InstancedActors/Public/InstancedActorsManager.h
+++ b/Engine/Plugins/Runtime/InstancedActors/Source/InstancedActors/Public/InstancedActorsManager.h
@@ -243,6 +243,9 @@ protected:
 	virtual void RequestActorSave(AActor* Actor) {}
 
 	//~ Begin AActor Overrides
+#if UE_WITH_IRIS
+	virtual void BeginReplication() override;
+#endif
 	virtual void BeginPlay() override;
 	virtual void EndPlay(const EEndPlayReason::Type EndPlayReason) override;
 	virtual bool IsHLODRelevant() const override;
