Fixed AInstancedActorsManager::HasInstancesOfClass erroneously reporting "no hit" when the entity being tested doesn't have an actor spawned to represent it.

diff --git a/Engine/Plugins/Runtime/InstancedActors/Source/InstancedActors/Private/InstancedActorsManager.cpp b/Engine/Plugins/Runtime/InstancedActors/Source/InstancedActors/Private/InstancedActorsManager.cpp
--- a/Engine/Plugins/Runtime/InstancedActors/Source/InstancedActors/Private/InstancedActorsManager.cpp
+++ b/Engine/Plugins/Runtime/InstancedActors/Source/InstancedActors/Private/InstancedActorsManager.cpp
@@ -1069,6 +1069,11 @@ bool AInstancedActorsManager::HasInstancesOfClass(const FBox& InQueryBounds, TSu
 						}
 					}
 				}
+				else
+				{
+					// if the bounds overlap but there's no actor we need to say it's a hit, since we cannot rule that out.
+					bHasInstance = true;
+				}
 			}
 			else 
 			{
diff --git a/Engine/Plugins/Runtime/InstancedActors/Source/InstancedActors/Private/InstancedActorsSubsystem.cpp b/Engine/Plugins/Runtime/InstancedActors/Source/InstancedActors/Private/InstancedActorsSubsystem.cpp
--- a/Engine/Plugins/Runtime/InstancedActors/Source/InstancedActors/Private/InstancedActorsSubsystem.cpp
+++ b/Engine/Plugins/Runtime/InstancedActors/Source/InstancedActors/Private/InstancedActorsSubsystem.cpp
@@ -19,6 +19,7 @@
 #include "MassEntitySubsystem.h"
 #include "Misc/ArchiveMD5.h"
 #include "Misc/ReverseIterate.h"
+#include "VisualLogger/VisualLogger.h"
 
 #if WITH_EDITOR
 #include "Logging/MessageLog.h"
@@ -578,6 +579,8 @@ void UInstancedActorsSubsystem::ForEachInstance(const FBox& QueryBounds, TFuncti
 bool UInstancedActorsSubsystem::HasInstancesOfClass(const FBox& QueryBounds, TSubclassOf<AActor> ActorClass, const bool bTestActorsIfSpawned) const
 {
 	TRACE_CPUPROFILER_EVENT_SCOPE(UInstancedActorsSubsystem::HasInstancesOfClass);
+	
+	UE_VLOG_BOX(this, LogInstancedActors, Log, QueryBounds, FColor::Orange, TEXT(""));
 
 	bool bHasInstances = false;
 	ForEachManager(QueryBounds, [QueryBounds, ActorClass, &bHasInstances, bTestActorsIfSpawned](AInstancedActorsManager& Manager)
@@ -897,4 +900,4 @@ bool UInstancedActorsSubsystem::RegisterNewSharedFragmentsInternal(TConstStructV
 	check(SortedSharedFragments.Num() == AllSharedFragmentsOfType.Num());
 
 	return bParamStructFound;
-}
\ No newline at end of file
+}
