[World Partition] Merged AActor GetStreamingBounds and GetStreamingBoundsEditor to make sure all overrides returns an editor bounds.

diff --git a/Engine/Plugins/Experimental/Water/Source/Runtime/Private/WaterZoneActor.cpp b/Engine/Plugins/Experimental/Water/Source/Runtime/Private/WaterZoneActor.cpp
--- a/Engine/Plugins/Experimental/Water/Source/Runtime/Private/WaterZoneActor.cpp
+++ b/Engine/Plugins/Experimental/Water/Source/Runtime/Private/WaterZoneActor.cpp
@@ -420,9 +420,9 @@ TUniquePtr<class FWorldPartitionActorDesc> AWaterZone::CreateClassActorDesc() co
 	return TUniquePtr<FWorldPartitionActorDesc>(new FWaterZoneActorDesc());
 }
 
-FBox AWaterZone::GetStreamingBounds() const
+void AWaterZone::GetStreamingBounds(FBox& OutRuntimeBounds, FBox& OutEditorBounds) const
 {
-	return GetZoneBounds();
+	OutRuntimeBounds = OutEditorBounds = GetZoneBounds();
 }
 #endif // WITH_EDITOR
 
diff --git a/Engine/Plugins/Experimental/Water/Source/Runtime/Public/WaterZoneActor.h b/Engine/Plugins/Experimental/Water/Source/Runtime/Public/WaterZoneActor.h
--- a/Engine/Plugins/Experimental/Water/Source/Runtime/Public/WaterZoneActor.h
+++ b/Engine/Plugins/Experimental/Water/Source/Runtime/Public/WaterZoneActor.h
@@ -93,7 +93,7 @@ public:
 	
 #if WITH_EDITOR
 	virtual TUniquePtr<class FWorldPartitionActorDesc> CreateClassActorDesc() const override;
-	virtual FBox GetStreamingBounds() const override;
+	virtual void GetStreamingBounds(FBox& OutRuntimeBounds, FBox& OutEditorBounds) const override;
 #endif //WITH_EDITOR
 
 	UFUNCTION(BlueprintCallable, Category=Rendering)
diff --git a/Engine/Plugins/Runtime/InstancedActors/Source/InstancedActors/Private/InstancedActorsManager.cpp b/Engine/Plugins/Runtime/InstancedActors/Source/InstancedActors/Private/InstancedActorsManager.cpp
--- a/Engine/Plugins/Runtime/InstancedActors/Source/InstancedActors/Private/InstancedActorsManager.cpp
+++ b/Engine/Plugins/Runtime/InstancedActors/Source/InstancedActors/Private/InstancedActorsManager.cpp
@@ -608,7 +608,7 @@ void AInstancedActorsManager::SerializeInstancePersistenceData(FStructuredArchiv
 }
 
 #if WITH_EDITOR
-FBox AInstancedActorsManager::GetStreamingBounds() const
+void AInstancedActorsManager::GetStreamingBounds(FBox& OutRuntimeBounds, FBox& OutEditorBounds) const
 {
 	UWorld* World = GetWorld();
 	check(World);
@@ -616,15 +616,10 @@ FBox AInstancedActorsManager::GetStreamingBounds() const
 	if (World->IsPartitionedWorld())
 	{
 		UActorPartitionSubsystem::FCellCoord CellCoord = UActorPartitionSubsystem::FCellCoord::GetCellCoord(GetActorLocation(), GetLevel(), GetGridSize());
-		return UActorPartitionSubsystem::FCellCoord::GetCellBounds(CellCoord, GetGridSize());
-	}
-	else
-	{
-		constexpr double BoundsMin = TNumericLimits<double>::Lowest();
-		constexpr double BoundsMax = TNumericLimits<double>::Max();
-		FBox LocalCellBounds = { FVector(BoundsMin), FVector(BoundsMax) };
-		return LocalCellBounds;
+		OutRuntimeBounds = OutEditorBounds = UActorPartitionSubsystem::FCellCoord::GetCellBounds(CellCoord, GetGridSize());
 	}
+
+	OutRuntimeBounds = OutEditorBounds = FBox(FVector(-HALF_WORLD_MAX), FVector(HALF_WORLD_MAX));
 }
 
 uint32 AInstancedActorsManager::GetDefaultGridSize(UWorld* InWorld) const
@@ -1130,8 +1125,10 @@ void AInstancedActorsManager::AuditInstances(FOutputDevice& Ar, bool bDebugDraw,
 
 #if WITH_EDITOR
 		// Draw streaming bounds
-		FBox StreamingBounds = GetStreamingBounds();
-		DrawDebugBox(World, StreamingBounds.GetCenter(), StreamingBounds.GetExtent(), FColor::Orange, /*bPersistentLines*/ DebugDrawDuration == -1.0f, DebugDrawDuration);
+		FBox RuntimeBounds;
+		FBox EditorBounds;
+		GetStreamingBounds(RuntimeBounds, EditorBounds);
+		DrawDebugBox(World, RuntimeBounds.GetCenter(), RuntimeBounds.GetExtent(), FColor::Orange, /*bPersistentLines*/ DebugDrawDuration == -1.0f, DebugDrawDuration);
 #endif // WITH_EDITOR
 	}
 #endif // UE_ENABLE_DEBUG_DRAWING
@@ -1817,14 +1814,16 @@ bool AInstancedActorsManager::SetSMInstanceTransform(const FSMInstanceId& Instan
 #endif
 
 	// Don't allow people to move an IA outside of this IAMs streaming bounds.
-	const FBox StreamingBounds = GetStreamingBounds();
+	FBox RuntimeBounds;
+	FBox EditorBounds;
+	GetStreamingBounds(RuntimeBounds, EditorBounds);
 	FTransform WorldTransform = bWorldSpace ? InstanceTransform : (InstanceTransform * GetActorTransform());
-	if (!StreamingBounds.IsInside(WorldTransform.GetLocation()))
+	if (!RuntimeBounds.IsInside(WorldTransform.GetLocation()))
 	{
 #if ENABLE_DRAW_DEBUG
 		FColor BoundsColor(255, 20, 20, 125);
-		DrawDebugSolidBox(GetWorld(), StreamingBounds, BoundsColor, FTransform::Identity);
-		DrawDebugBox(GetWorld(), StreamingBounds.GetCenter(), StreamingBounds.GetExtent(), FColor::Red);
+		DrawDebugSolidBox(GetWorld(), RuntimeBounds, BoundsColor, FTransform::Identity);
+		DrawDebugBox(GetWorld(), RuntimeBounds.GetCenter(), RuntimeBounds.GetExtent(), FColor::Red);
 #endif
 		return false;
 	}
diff --git a/Engine/Plugins/Runtime/InstancedActors/Source/InstancedActors/Public/InstancedActorsManager.h b/Engine/Plugins/Runtime/InstancedActors/Source/InstancedActors/Public/InstancedActorsManager.h
--- a/Engine/Plugins/Runtime/InstancedActors/Source/InstancedActors/Public/InstancedActorsManager.h
+++ b/Engine/Plugins/Runtime/InstancedActors/Source/InstancedActors/Public/InstancedActorsManager.h
@@ -253,7 +253,7 @@ protected:
 	virtual void PostLoad() override;
 #if WITH_EDITOR
 	virtual bool IsUserManaged() const override { return true; }
-	virtual FBox GetStreamingBounds() const override;
+	virtual void GetStreamingBounds(FBox& OutRuntimeBounds, FBox& OutEditorBounds) const override;
 #endif
 	//~ End AActor Overrides
 
diff --git a/Engine/Source/Editor/UnrealEd/Classes/Editor/GroupActor.h b/Engine/Source/Editor/UnrealEd/Classes/Editor/GroupActor.h
--- a/Engine/Source/Editor/UnrealEd/Classes/Editor/GroupActor.h
+++ b/Engine/Source/Editor/UnrealEd/Classes/Editor/GroupActor.h
@@ -46,7 +46,7 @@ class AGroupActor : public AActor
 	virtual void SetIsTemporarilyHiddenInEditor( bool bIsHidden ) override;
 	virtual void GetActorBounds(bool bOnlyCollidingComponents, FVector& Origin, FVector& BoxExtent, bool bIncludeFromChildActors) const override;
 #if WITH_EDITOR
-	virtual FBox GetStreamingBounds() const override;
+	virtual void GetStreamingBounds(FBox& OutRuntimeBounds, FBox& OutEditorBounds) const override;
 #endif
 	//~ End AActor Interface
 
diff --git a/Engine/Source/Editor/UnrealEd/Private/Commandlets/WorldPartitionConvertCommandlet.cpp b/Engine/Source/Editor/UnrealEd/Private/Commandlets/WorldPartitionConvertCommandlet.cpp
--- a/Engine/Source/Editor/UnrealEd/Private/Commandlets/WorldPartitionConvertCommandlet.cpp
+++ b/Engine/Source/Editor/UnrealEd/Private/Commandlets/WorldPartitionConvertCommandlet.cpp
@@ -1016,9 +1016,11 @@ int32 UWorldPartitionConvertCommandlet::Main(const FString& Params)
 					// Only override default grid placement on actors that are spatially loaded
 					else if (Actor->GetIsSpatiallyLoaded() && Actor->CanChangeIsSpatiallyLoadedFlag())
 					{
-						const FBox ActorBounds = Actor->GetStreamingBounds();
+						FBox ActorRuntimeBounds;
+						FBox ActorEditorBounds;
+						Actor->GetStreamingBounds(ActorRuntimeBounds, ActorEditorBounds);
 
-						if (!WorldBounds.IsInside(ActorBounds))
+						if (!WorldBounds.IsInside(ActorRuntimeBounds))
 						{
 							Actor->SetIsSpatiallyLoaded(false);
 						}
diff --git a/Engine/Source/Editor/UnrealEd/Private/GroupActor.cpp b/Engine/Source/Editor/UnrealEd/Private/GroupActor.cpp
--- a/Engine/Source/Editor/UnrealEd/Private/GroupActor.cpp
+++ b/Engine/Source/Editor/UnrealEd/Private/GroupActor.cpp
@@ -251,15 +251,19 @@ void AGroupActor::GetActorBounds(bool bOnlyCollidingComponents, FVector& Origin,
 }
 
 #if WITH_EDITOR
-FBox AGroupActor::GetStreamingBounds() const
+void AGroupActor::GetStreamingBounds(FBox& OutRuntimeBounds, FBox& OutEditorBounds) const
 {
-	FBox StreamingBounds = Super::GetStreamingBounds();
+	Super::GetStreamingBounds(OutRuntimeBounds, OutEditorBounds);
 
 	for (AActor* Actor : GroupActors)
 	{
 		if (Actor)
 		{
-			StreamingBounds += Actor->GetStreamingBounds();
+			FBox RuntimeBounds;
+			FBox EditorBounds;
+			Actor->GetStreamingBounds(RuntimeBounds, EditorBounds);
+			OutRuntimeBounds += RuntimeBounds;
+			OutEditorBounds += EditorBounds;
 		}
 	}
 
@@ -267,11 +271,14 @@ FBox AGroupActor::GetStreamingBounds() const
 	{
 		if (SubGroupActor)
 		{
-			StreamingBounds += SubGroupActor->GetStreamingBounds();
+			FBox RuntimeBounds;
+			FBox EditorBounds;
+			SubGroupActor->GetStreamingBounds(RuntimeBounds, EditorBounds);
+			OutRuntimeBounds += RuntimeBounds;
+			OutEditorBounds += EditorBounds;
+
 		}
 	}
-
-	return StreamingBounds;
 }
 #endif
 
diff --git a/Engine/Source/Editor/WorldPartitionEditor/Private/WorldPartition/SWorldPartitionEditorGrid2D.cpp b/Engine/Source/Editor/WorldPartitionEditor/Private/WorldPartition/SWorldPartitionEditorGrid2D.cpp
--- a/Engine/Source/Editor/WorldPartitionEditor/Private/WorldPartition/SWorldPartitionEditorGrid2D.cpp
+++ b/Engine/Source/Editor/WorldPartitionEditor/Private/WorldPartition/SWorldPartitionEditorGrid2D.cpp
@@ -1400,7 +1400,8 @@ uint32 SWorldPartitionEditorGrid2D::PaintActors(const FGeometry& AllottedGeometr
 			{
 				Guid = InActor->GetActorGuid();
 				Label = *InActor->GetActorLabel(false);
-				ActorBounds = InActor->GetStreamingBoundsEditor();
+				FBox RuntimeActorBounds;
+				InActor->GetStreamingBounds(RuntimeActorBounds, ActorBounds);
 				DescBounds = InActorDescInstance ? InActorDescInstance->GetEditorBounds() : ActorBounds;
 				bIsSpatiallyLoaded = InActor->GetIsSpatiallyLoaded();
 			}
@@ -2194,7 +2195,10 @@ void SWorldPartitionEditorGrid2D::FocusSelection()
 	{
 		if (AActor* Actor = Cast<AActor>(*It))
 		{
-			SelectionBox += Actor->GetStreamingBoundsEditor();
+			FBox RuntimeBounds;
+			FBox EditorBounds;
+			Actor->GetStreamingBounds(RuntimeBounds, EditorBounds);
+			SelectionBox += EditorBounds;
 		}
 	}
 
diff --git a/Engine/Source/Runtime/Engine/Classes/Engine/DirectionalLight.h b/Engine/Source/Runtime/Engine/Classes/Engine/DirectionalLight.h
--- a/Engine/Source/Runtime/Engine/Classes/Engine/DirectionalLight.h
+++ b/Engine/Source/Runtime/Engine/Classes/Engine/DirectionalLight.h
@@ -38,7 +38,7 @@ public:
 	virtual void LoadedFromAnotherClass(const FName& OldClassName) override;
 	virtual void PostEditChangeProperty(FPropertyChangedEvent& PropertyChangedEvent) override;
 	virtual bool CanChangeIsSpatiallyLoadedFlag() const override { return false; }
-	virtual FBox GetStreamingBounds() const override { return FBox(ForceInit); }
+	virtual void GetStreamingBounds(FBox& OutRuntimeBounds, FBox& OutEditorBounds) const override { OutRuntimeBounds.Init(); OutEditorBounds.Init(); }
 #endif
 
 public:
diff --git a/Engine/Source/Runtime/Engine/Classes/GameFramework/Actor.h b/Engine/Source/Runtime/Engine/Classes/GameFramework/Actor.h
--- a/Engine/Source/Runtime/Engine/Classes/GameFramework/Actor.h
+++ b/Engine/Source/Runtime/Engine/Classes/GameFramework/Actor.h
@@ -2443,25 +2443,17 @@ public:
 	}
 
 	/**
-	 * Returns the bounding box of all components that make up this Actor for loading at runtime.
+	 * Returns the bounding boxes of all components that make up this Actor for loading at runtime and editor.
 	 *
 	 * This function differs from GetActorBounds because it will return a valid origin and an empty extent if this actor
 	 * doesn't have primitive components.
 	 *
 	 * @see GetActorBounds()
 	 */
-	ENGINE_API virtual FBox GetStreamingBounds() const;
-
-	/**
-	 * Returns the bounding box of all components that make up this Actor for loading in the editor.
-	 *
-	 * This function differs from GetActorBounds because it will return a valid origin and an empty extent if this actor
-	 * doesn't have primitive components.
-	 *
-	 * @see GetActorBounds()
-	 */
-	ENGINE_API virtual FBox GetStreamingBoundsEditor() const;
+	ENGINE_API virtual void GetStreamingBounds(FBox& OutRuntimeBounds, FBox& OutEditorBounds) const;
 
+	UE_DEPRECATED(5.5, "Use the override that takes both Runtime and Editor boxes")
+	ENGINE_API FBox GetStreamingBounds() const;
 #endif
 
 
diff --git a/Engine/Source/Runtime/Engine/Private/Actor.cpp b/Engine/Source/Runtime/Engine/Private/Actor.cpp
--- a/Engine/Source/Runtime/Engine/Private/Actor.cpp
+++ b/Engine/Source/Runtime/Engine/Private/Actor.cpp
@@ -1214,20 +1214,25 @@ static bool HasComponentForceActorNonSpatiallyLoaded(const AActor* InActor)
 	return bHasComponentForceActorNonSpatiallyLoaded;
 }
 
-FBox AActor::GetStreamingBounds() const
+void AActor::GetStreamingBounds(FBox& OutRuntimeBounds, FBox& OutEditorBounds) const
 {
-	FBox StreamingBounds(ForceInit);
-	ForEachStreamingRelevantComponent(this, false, [&StreamingBounds](const UActorComponent* Component, const FBox& StreamingBound)
+	OutRuntimeBounds.Init();
+	ForEachStreamingRelevantComponent(this, false, [&OutRuntimeBounds](const UActorComponent* Component, const FBox& StreamingBound)
 	{
-		StreamingBounds += StreamingBound;
+		OutRuntimeBounds += StreamingBound;
+	});
+
+	OutEditorBounds.Init();
+	ForEachStreamingRelevantComponent(this, true, [&OutEditorBounds](const UActorComponent* Component, const FBox& StreamingBound)
+	{
+		OutEditorBounds += StreamingBound;
 	});
-	return StreamingBounds;
 }
 
-FBox AActor::GetStreamingBoundsEditor() const
+FBox AActor::GetStreamingBounds() const
 {
 	FBox StreamingBounds(ForceInit);
-	ForEachStreamingRelevantComponent(this, true, [&StreamingBounds](const UActorComponent* Component, const FBox& StreamingBound)
+	ForEachStreamingRelevantComponent(this, false, [&StreamingBounds](const UActorComponent* Component, const FBox& StreamingBound)
 	{
 		StreamingBounds += StreamingBound;
 	});
diff --git a/Engine/Source/Runtime/Engine/Private/LevelInstance/LevelInstanceActor.cpp b/Engine/Source/Runtime/Engine/Private/LevelInstance/LevelInstanceActor.cpp
--- a/Engine/Source/Runtime/Engine/Private/LevelInstance/LevelInstanceActor.cpp
+++ b/Engine/Source/Runtime/Engine/Private/LevelInstance/LevelInstanceActor.cpp
@@ -282,26 +282,21 @@ void ALevelInstance::EditorGetUnderlyingActors(TSet<AActor*>& OutUnderlyingActor
 	LevelInstanceActorImpl.EditorGetUnderlyingActors(OutUnderlyingActors);
 }
 
-FBox ALevelInstance::GetStreamingBounds() const
+void ALevelInstance::GetStreamingBounds(FBox& OutRuntimeBounds, FBox& OutEditorBounds) const
 {
-	FBox LevelInstanceBounds;
-	if (LevelInstanceActorImpl.GetBounds(LevelInstanceBounds))
+	Super::GetStreamingBounds(OutRuntimeBounds, OutEditorBounds);
+	
+	FBox LevelInstanceRuntimeBounds;
+	if (LevelInstanceActorImpl.GetBounds(LevelInstanceRuntimeBounds))
 	{
-		return LevelInstanceBounds;
+		OutRuntimeBounds = LevelInstanceRuntimeBounds;
 	}
 
-	return Super::GetStreamingBounds();
-}
-
-FBox ALevelInstance::GetStreamingBoundsEditor() const
-{
 	FBox LevelInstanceEditorBounds;
 	if (LevelInstanceActorImpl.GetEditorBounds(LevelInstanceEditorBounds))
 	{
-		return LevelInstanceEditorBounds;
+		OutEditorBounds = LevelInstanceEditorBounds;
 	}
-
-	return Super::GetStreamingBoundsEditor();
 }
 
 bool ALevelInstance::IsUserManaged() const
diff --git a/Engine/Source/Runtime/Engine/Private/WorldPartition/HLOD/HLODActor.cpp b/Engine/Source/Runtime/Engine/Private/WorldPartition/HLOD/HLODActor.cpp
--- a/Engine/Source/Runtime/Engine/Private/WorldPartition/HLOD/HLODActor.cpp
+++ b/Engine/Source/Runtime/Engine/Private/WorldPartition/HLOD/HLODActor.cpp
@@ -480,9 +480,9 @@ void AWorldPartitionHLOD::SetHLODBounds(const FBox& InBounds)
 	HLODBounds = InBounds;
 }
 
-FBox AWorldPartitionHLOD::GetStreamingBounds() const
+void AWorldPartitionHLOD::GetStreamingBounds(FBox& OutRuntimeBounds, FBox& OutEditorBounds) const
 {
-	return HLODBounds;
+	OutRuntimeBounds = OutEditorBounds = HLODBounds;
 }
 
 int64 AWorldPartitionHLOD::GetStat(FName InStatName) const
diff --git a/Engine/Source/Runtime/Engine/Private/WorldPartition/Landscape/LandscapeActorDesc.cpp b/Engine/Source/Runtime/Engine/Private/WorldPartition/Landscape/LandscapeActorDesc.cpp
--- a/Engine/Source/Runtime/Engine/Private/WorldPartition/Landscape/LandscapeActorDesc.cpp
+++ b/Engine/Source/Runtime/Engine/Private/WorldPartition/Landscape/LandscapeActorDesc.cpp
@@ -25,7 +25,9 @@ void FLandscapeActorDesc::Init(const AActor* InActor)
 			LandscapeActorGuid = LandscapeActor->GetActorGuid();
 		}
 
-		EditorBounds = InActor->GetStreamingBoundsEditor();
+		// FLandscapeActorDesc derives from FPartitionActorDesc but doesn't use the cell bounds as the parent class was designed for.
+		// @todo_ow: make FLandscapeActorDesc derives from FWorldPartitionActorDesc instead?
+		InActor->GetStreamingBounds(RuntimeBounds, EditorBounds);
 	}
 }
 
diff --git a/Engine/Source/Runtime/Engine/Private/WorldPartition/NavigationData/NavigationDataChunkActor.cpp b/Engine/Source/Runtime/Engine/Private/WorldPartition/NavigationData/NavigationDataChunkActor.cpp
--- a/Engine/Source/Runtime/Engine/Private/WorldPartition/NavigationData/NavigationDataChunkActor.cpp
+++ b/Engine/Source/Runtime/Engine/Private/WorldPartition/NavigationData/NavigationDataChunkActor.cpp
@@ -146,9 +146,9 @@ void ANavigationDataChunkActor::GetActorBounds(bool bOnlyCollidingComponents, FV
 }
 
 #if WITH_EDITOR
-FBox ANavigationDataChunkActor::GetStreamingBounds() const
+void ANavigationDataChunkActor::GetStreamingBounds(FBox& OutRuntimeBounds, FBox& OutEditorBounds) const
 {
-	return DataChunkActorBounds;
+	OutRuntimeBounds = OutEditorBounds = DataChunkActorBounds;
 }
 #endif // WITH_EDITOR
 
diff --git a/Engine/Source/Runtime/Engine/Private/WorldPartition/WorldPartitionActorDesc.cpp b/Engine/Source/Runtime/Engine/Private/WorldPartition/WorldPartitionActorDesc.cpp
--- a/Engine/Source/Runtime/Engine/Private/WorldPartition/WorldPartitionActorDesc.cpp
+++ b/Engine/Source/Runtime/Engine/Private/WorldPartition/WorldPartitionActorDesc.cpp
@@ -90,8 +90,14 @@ void FWorldPartitionActorDesc::Init(const AActor* InActor)
 
 	ActorTransform = InActor->GetActorTransform();
 
-	RuntimeBounds = !bIsDefaultActorDesc ? InActor->GetStreamingBounds() : FBox(ForceInit);
-	EditorBounds = !bIsDefaultActorDesc ? InActor->GetStreamingBoundsEditor() : FBox(ForceInit);
+	RuntimeBounds.Init();
+	EditorBounds.Init();
+	
+	if (!bIsDefaultActorDesc)
+	{
+		InActor->GetStreamingBounds(RuntimeBounds, EditorBounds);
+	}
+
 	RuntimeGrid = InActor->GetRuntimeGrid();
 	bIsSpatiallyLoaded = InActor->GetIsSpatiallyLoaded();
 	bActorIsEditorOnly = InActor->IsEditorOnly();
@@ -605,6 +611,11 @@ void FWorldPartitionActorDesc::Serialize(FArchive& Ar)
 			Ar << RuntimeBounds << EditorBounds;
 		}
 	}
+	else
+	{
+		RuntimeBounds.Init();
+		EditorBounds.Init();
+	}
 	
 	if (Ar.CustomVer(FUE5ReleaseStreamObjectVersion::GUID) < FUE5ReleaseStreamObjectVersion::ConvertedActorGridPlacementToSpatiallyLoadedFlag)
 	{
diff --git a/Engine/Source/Runtime/Engine/Public/LevelInstance/LevelInstanceActor.h b/Engine/Source/Runtime/Engine/Public/LevelInstance/LevelInstanceActor.h
--- a/Engine/Source/Runtime/Engine/Public/LevelInstance/LevelInstanceActor.h
+++ b/Engine/Source/Runtime/Engine/Public/LevelInstance/LevelInstanceActor.h
@@ -102,8 +102,7 @@ public:
 	ENGINE_API virtual void PushSelectionToProxies() override;
 	ENGINE_API virtual void PushLevelInstanceEditingStateToProxies(bool bInEditingState) override;
 	ENGINE_API virtual FBox GetComponentsBoundingBox(bool bNonColliding = false, bool bIncludeFromChildActors = false) const override;
-	ENGINE_API virtual FBox GetStreamingBounds() const override;
-	ENGINE_API virtual FBox GetStreamingBoundsEditor() const override;
+	ENGINE_API virtual void GetStreamingBounds(FBox& OutRuntimeBounds, FBox& OutEditorBounds) const override;
 	ENGINE_API virtual bool IsLockLocation() const override;
 	ENGINE_API virtual bool IsActorLabelEditable() const override;
 	ENGINE_API virtual bool GetReferencedContentObjects(TArray<UObject*>& Objects) const override;
diff --git a/Engine/Source/Runtime/Engine/Public/WorldPartition/HLOD/HLODActor.h b/Engine/Source/Runtime/Engine/Public/WorldPartition/HLOD/HLODActor.h
--- a/Engine/Source/Runtime/Engine/Public/WorldPartition/HLOD/HLODActor.h
+++ b/Engine/Source/Runtime/Engine/Public/WorldPartition/HLOD/HLODActor.h
@@ -93,7 +93,7 @@ protected:
 	virtual bool IsRuntimeOnly() const override { return true; }
 #if WITH_EDITOR
 	ENGINE_API virtual TUniquePtr<class FWorldPartitionActorDesc> CreateClassActorDesc() const override;
-	ENGINE_API virtual FBox GetStreamingBounds() const override;
+	ENGINE_API virtual void GetStreamingBounds(FBox& OutRuntimeBounds, FBox& OutEditorBounds) const override;
 
 	virtual bool ShouldImport(FStringView ActorPropString, bool IsMovingLevel) override { return false; }
 	virtual bool IsLockLocation() const override { return true; }
diff --git a/Engine/Source/Runtime/Engine/Public/WorldPartition/NavigationData/NavigationDataChunkActor.h b/Engine/Source/Runtime/Engine/Public/WorldPartition/NavigationData/NavigationDataChunkActor.h
--- a/Engine/Source/Runtime/Engine/Public/WorldPartition/NavigationData/NavigationDataChunkActor.h
+++ b/Engine/Source/Runtime/Engine/Public/WorldPartition/NavigationData/NavigationDataChunkActor.h
@@ -38,7 +38,7 @@ public:
 
 	//~ Begin AActor Interface.
 	virtual bool CanChangeIsSpatiallyLoadedFlag() const override { return false; }
-	ENGINE_API virtual FBox GetStreamingBounds() const override;
+	ENGINE_API virtual void GetStreamingBounds(FBox& OutRuntimeBounds, FBox& OutEditorBounds) const override;
 	//~ End AActor Interface.
 #endif // WITH_EDITOR
 
