Minor optimization moving InstancedActorLocationQuery out of AInstancedActorsManager::ForEachInstance and making it AInstancedActorsManager member. This will improve perf due to the query not re-caching valid archetypes every single time AInstancedActorsManager::ForEachInstance is called

diff --git a/Engine/Plugins/Runtime/InstancedActors/Source/InstancedActors/Private/InstancedActorsManager.cpp b/Engine/Plugins/Runtime/InstancedActors/Source/InstancedActors/Private/InstancedActorsManager.cpp
--- a/Engine/Plugins/Runtime/InstancedActors/Source/InstancedActors/Private/InstancedActorsManager.cpp
+++ b/Engine/Plugins/Runtime/InstancedActors/Source/InstancedActors/Private/InstancedActorsManager.cpp
@@ -243,6 +243,12 @@ void AInstancedActorsManager::InitializeModifyAndSpawnEntities()
 {
 	check(IsValid(InstancedActorSubsystem));
 
+	if (InstancedActorLocationQuery.IsEmpty())
+	{
+		InstancedActorLocationQuery.AddRequirement<FInstancedActorsFragment>(EMassFragmentAccess::ReadOnly);
+		InstancedActorLocationQuery.AddRequirement<FTransformFragment>(EMassFragmentAccess::ReadOnly);
+	}
+
 	// Debug draw managers as they load
 #if WITH_INSTANCEDACTORS_DEBUG
 	UE::InstancedActors::Debug::DebugDrawManager(UE::InstancedActors::Debug::CVars::DebugManagerLoading, *this);
@@ -798,10 +804,6 @@ bool AInstancedActorsManager::ForEachInstance(FInstanceOperationFunc Operation,
 	{
 		check(MassEntityManager.IsValid());
 
-		FMassEntityQuery InstancedActorLocationQuery;
-		InstancedActorLocationQuery.AddRequirement<FInstancedActorsFragment>(EMassFragmentAccess::ReadOnly);
-		InstancedActorLocationQuery.AddRequirement<FTransformFragment>(EMassFragmentAccess::ReadOnly);
-
 		for (TObjectPtr<UInstancedActorsData> InstanceData : PerActorClassInstanceData)
 		{
 			// InstancedActorDataPredicate filter PerActorClassInstanceData
diff --git a/Engine/Plugins/Runtime/InstancedActors/Source/InstancedActors/Public/InstancedActorsManager.h b/Engine/Plugins/Runtime/InstancedActors/Source/InstancedActors/Public/InstancedActorsManager.h
--- a/Engine/Plugins/Runtime/InstancedActors/Source/InstancedActors/Public/InstancedActorsManager.h
+++ b/Engine/Plugins/Runtime/InstancedActors/Source/InstancedActors/Public/InstancedActorsManager.h
@@ -5,6 +5,7 @@
 #include "Containers/ArrayView.h"
 #include "InstancedActorsTypes.h"
 #include "InstancedActorsIndex.h"
+#include "MassEntityQuery.h"
 #include "ActorPartition/PartitionActor.h"
 #include "Containers/BitArray.h"
 #include "Templates/SharedPointer.h"
@@ -350,6 +351,8 @@ protected:
 	UPROPERTY(EditAnywhere, Category=InstancedActor)
 	TSubclassOf<UInstancedActorsData> InstancedActorsDataClass;
 
+	mutable FMassEntityQuery InstancedActorLocationQuery;
+
 private:
 #if WITH_EDITORONLY_DATA
 	UPROPERTY()
diff --git a/Engine/Source/Runtime/MassEntity/Public/MassEntityQuery.h b/Engine/Source/Runtime/MassEntity/Public/MassEntityQuery.h
--- a/Engine/Source/Runtime/MassEntity/Public/MassEntityQuery.h
+++ b/Engine/Source/Runtime/MassEntity/Public/MassEntityQuery.h
@@ -88,6 +88,8 @@ public:
 
 	void RequireMutatingWorldAccess() { bRequiresMutatingWorldAccess = true; }
 
+	bool IsEmpty() const { return FMassFragmentRequirements::IsEmpty() && FMassSubsystemRequirements::IsEmpty(); }
+
 	const TArray<FMassArchetypeHandle>& GetArchetypes() const
 	{ 
 		return ValidArchetypes; 
