[IAM] Fix to prevent getting two sets of IAMS in PIE, one from the main level and one from the HLOD.
NOTE: This fix was provided by [REDACTED_HANDLE]. IAMs are sometimes automatically tagged as editor-only in the non-HLOD map and automatically cloned (without the editor-only flag) to the HLOD. This fix prevents PIE EditorOnly IAMs from beign spawned to avoid duplication.
#rnx

diff --git a/Engine/Plugins/Runtime/InstancedActors/Source/InstancedActors/Private/InstancedActorsManager.cpp b/Engine/Plugins/Runtime/InstancedActors/Source/InstancedActors/Private/InstancedActorsManager.cpp
--- a/Engine/Plugins/Runtime/InstancedActors/Source/InstancedActors/Private/InstancedActorsManager.cpp
+++ b/Engine/Plugins/Runtime/InstancedActors/Source/InstancedActors/Private/InstancedActorsManager.cpp
@@ -167,15 +167,22 @@ void AInstancedActorsManager::BeginReplication()
 
 void AInstancedActorsManager::BeginPlay()
 {
+	UWorld* World = GetWorld();
+	check(World);
+
+#if WITH_EDITORONLY_DATA
+	if (bIsEditorOnlyActor && World && World->IsPlayInEditor())
+	{
+		Destroy();
+		return;
+	}
+#endif
+
 	Super::BeginPlay();
 
 	TRACE_CPUPROFILER_EVENT_SCOPE("AInstancedActorsManager::BeginPlay");
 
 	UE_LOG(LogInstancedActors, Verbose, TEXT("%s (%d instances) BeginPlay"), *GetPathName(), GetNumValidInstances());
-
-	UWorld* World = GetWorld();
-	check(World);
-
 	UMassEntitySubsystem* EntitySubsystem = World->GetSubsystem<UMassEntitySubsystem>();
 	if (!ensure(EntitySubsystem))
 	{
