InstancedActors editor module providing "Convert To IA" and "From IA" actor menu functions.

The functionality moved oved from project-specific code and improved (supporyt for mixing IAM and regular actors in a Selection and handling accordingly).
The change also fixes the issue with not retaining the 3D scale when converting InstancedActors instances back to AActor instances.

diff --git a/Engine/Plugins/Runtime/InstancedActors/InstancedActors.uplugin b/Engine/Plugins/Runtime/InstancedActors/InstancedActors.uplugin
--- a/Engine/Plugins/Runtime/InstancedActors/InstancedActors.uplugin
+++ b/Engine/Plugins/Runtime/InstancedActors/InstancedActors.uplugin
@@ -23,6 +23,10 @@
 		{
 			"Name": "InstancedActorsTestSuite",
 			"Type": "UncookedOnly"
+		},
+		{
+			"Name": "InstancedActorsEditor",
+			"Type": "Editor"
 		}
 	],
 	"Plugins": [
diff --git a/Engine/Plugins/Runtime/InstancedActors/Source/InstancedActors/Private/InstancedActorsManager.cpp b/Engine/Plugins/Runtime/InstancedActors/Source/InstancedActors/Private/InstancedActorsManager.cpp
--- a/Engine/Plugins/Runtime/InstancedActors/Source/InstancedActors/Private/InstancedActorsManager.cpp
+++ b/Engine/Plugins/Runtime/InstancedActors/Source/InstancedActors/Private/InstancedActorsManager.cpp
@@ -697,6 +697,12 @@ UInstancedActorsData& AInstancedActorsManager::GetOrCreateActorInstanceData(TSub
 	FInstancedActorsVisualizationDesc DefaultVisualiation = EditorInstancedActorSubsystem->CreateVisualDescriptionFromActor(ExemplarActor);
 	CreateISMComponents(DefaultVisualiation, NewInstanceData->SharedSettings, NewInstanceData->EditorPreviewISMComponents, /*bEditorPreviewISMCs*/ true);
 
+	// Add EditorPreviewISMComponents to the known ISMComponent-to-InstanceData map (ISMComponentToInstanceDataMap)
+	// so that we can identify Instanced Actors instances by ISM instance id coming from one of the EditorPreviewISMComponents
+	// (like when the instance is clicked in the editor). 
+	// @see ActorInstanceHandleFromFSMInstanceId 
+	RegisterInstanceDatasComponents(*NewInstanceData, NewInstanceData->EditorPreviewISMComponents);
+
 	// Cache asset bounds for use during instance population
 	NewInstanceData->AssetBounds = CalculateBounds(ActorClass);
 	ensure(NewInstanceData->AssetBounds.IsValid);
@@ -706,8 +712,6 @@ UInstancedActorsData& AInstancedActorsManager::GetOrCreateActorInstanceData(TSub
 
 void AInstancedActorsManager::PreRegisterAllComponents()
 {
-	// This whole loop is pointless without editor only data, EditorPreviewISMComponents is editor only.
-#if WITH_EDITORONLY_DATA
 	for (UInstancedActorsData* InstanceData : PerActorClassInstanceData)
 	{
 		// Modify the ISMCs here since the components might have been serialized before this change.
@@ -719,11 +723,12 @@ void AInstancedActorsManager::PreRegisterAllComponents()
 			ISMComponent->bHasPerInstanceHitProxies = true;
 		}
 
-		// Add EditorPreviewISMComponents to ISMComponentToInstanceDataMap for ISMC and instance index to instance handle
-		// resolution in ActorInstanceHandleFromFSMInstanceId.
+		// Add EditorPreviewISMComponents to the known ISMComponent-to-InstanceData map (ISMComponentToInstanceDataMap)
+		// so that we can identify Instanced Actors instances by ISM instance id coming from one of the EditorPreviewISMComponents
+		// (like when the instance is clicked in the editor). 
+		// @see ActorInstanceHandleFromFSMInstanceId 
 		RegisterInstanceDatasComponents(*InstanceData, InstanceData->EditorPreviewISMComponents);
 	}
-#endif
 }
 
 FInstancedActorsInstanceHandle AInstancedActorsManager::AddActorInstance(TSubclassOf<AActor> ActorClass, FTransform InstanceTransform, bool bWorldSpace, const FInstancedActorsTagSet& InstanceTags)
diff --git a/Engine/Plugins/Runtime/InstancedActors/Source/InstancedActorsEditor/InstancedActorsEditor.Build.cs b/Engine/Plugins/Runtime/InstancedActors/Source/InstancedActorsEditor/InstancedActorsEditor.Build.cs
new file mode 100644
--- /dev/null
+++ b/Engine/Plugins/Runtime/InstancedActors/Source/InstancedActorsEditor/InstancedActorsEditor.Build.cs
@@ -0,0 +1,31 @@
+// Copyright Epic Games, Inc. All Rights Reserved.
+
+namespace UnrealBuildTool.Rules
+{
+	public class InstancedActorsEditor : ModuleRules
+	{
+		public InstancedActorsEditor(ReadOnlyTargetRules Target) : base(Target)
+		{
+			PCHUsage = ModuleRules.PCHUsageMode.UseExplicitOrSharedPCHs;
+
+			PublicIncludePaths.AddRange(
+				new string[] {
+				}
+			);
+
+			PublicDependencyModuleNames.AddRange(
+				new string[] {
+					"Core",
+					"CoreUObject",
+					"Engine",
+					"UnrealEd",
+					"LevelEditor",
+					"SlateCore",
+					"Slate",
+					"LevelEditor",
+					"InstancedActors",
+				}
+			);
+		}
+	}
+}
diff --git a/Engine/Plugins/Runtime/InstancedActors/Source/InstancedActorsEditor/Private/InstancedActorsEditorModule.cpp b/Engine/Plugins/Runtime/InstancedActors/Source/InstancedActorsEditor/Private/InstancedActorsEditorModule.cpp
new file mode 100644
--- /dev/null
+++ b/Engine/Plugins/Runtime/InstancedActors/Source/InstancedActorsEditor/Private/InstancedActorsEditorModule.cpp
@@ -0,0 +1,209 @@
+// Copyright Epic Games, Inc. All Rights Reserved.
+
+#include "InstancedActorsEditorModule.h"
+#include "InstancedActorsSubsystem.h"
+#include "InstancedActorsManager.h"
+#include "InstancedActorsData.h"
+#include "InstancedActorsIteration.h"
+#include "Modules/ModuleManager.h"
+#include "LevelEditor.h"
+#include "Engine/World.h"
+#include "ScopedTransaction.h"
+#include "ActorFactories/ActorFactory.h"
+
+
+#define LOCTEXT_NAMESPACE "InstancedActorsEditor"
+
+DEFINE_LOG_CATEGORY_STATIC(LogInstancedActorsEditor, Log, All);
+
+IMPLEMENT_MODULE(FInstancedActorsEditorModule, InstancedActorsEditor)
+
+void FInstancedActorsEditorModule::StartupModule()
+{
+	AddLevelViewportMenuExtender();
+}
+
+void FInstancedActorsEditorModule::ShutdownModule()
+{
+	// Cleanup menu extenstions
+	RemoveLevelViewportMenuExtender();
+}
+
+void FInstancedActorsEditorModule::AddLevelViewportMenuExtender()
+{
+	if (!IsRunningGame())
+	{
+		FLevelEditorModule& LevelEditorModule = FModuleManager::Get().LoadModuleChecked<FLevelEditorModule>("LevelEditor");
+		TArray<FLevelEditorModule::FLevelViewportMenuExtender_SelectedActors>& MenuExtenders = LevelEditorModule.GetAllLevelViewportContextMenuExtenders();
+
+		MenuExtenders.Add(FLevelEditorModule::FLevelViewportMenuExtender_SelectedActors::CreateRaw(this, &FInstancedActorsEditorModule::CreateLevelViewportContextMenuExtender));
+		LevelViewportExtenderHandle = MenuExtenders.Last().GetHandle();
+	}
+}
+
+void FInstancedActorsEditorModule::RemoveLevelViewportMenuExtender()
+{
+	if (LevelViewportExtenderHandle.IsValid())
+	{
+		FLevelEditorModule* LevelEditorModule = FModuleManager::Get().GetModulePtr<FLevelEditorModule>("LevelEditor");
+		if (LevelEditorModule)
+		{
+			LevelEditorModule->GetAllLevelViewportContextMenuExtenders().RemoveAll([this](const FLevelEditorModule::FLevelViewportMenuExtender_SelectedActors& In)
+				{ 
+					return In.GetHandle() == LevelViewportExtenderHandle; 
+				});
+		}
+		LevelViewportExtenderHandle.Reset();
+	}
+}
+
+TSharedRef<FExtender> FInstancedActorsEditorModule::CreateLevelViewportContextMenuExtender(const TSharedRef<FUICommandList> CommandList, const TArray<AActor*> InActors)
+{
+	TSharedRef<FExtender> Extender = MakeShareable(new FExtender);
+	if (InActors.Num() > 0)
+	{
+		FText ActorName = InActors.Num() == 1 ? FText::Format(LOCTEXT("ActorNameSingular", "\"{0}\""), FText::FromString(InActors[0]->GetActorLabel())) : LOCTEXT("ActorNamePlural", "Actors");
+
+		FLevelEditorModule& LevelEditor = FModuleManager::GetModuleChecked<FLevelEditorModule>(TEXT("LevelEditor"));
+		TSharedRef<FUICommandList> LevelEditorCommandBindings = LevelEditor.GetGlobalLevelEditorActions();
+
+		// Note: ActorConvert extension point appears only in the pulldown Actor menu.
+		Extender->AddMenuExtension("ActorConvert", EExtensionHook::After, LevelEditorCommandBindings, FMenuExtensionDelegate::CreateLambda(
+			[this, ActorName, InActors](FMenuBuilder& MenuBuilder)
+			{
+				bool bCanExecuteActorToIAM = false;
+				bool bCanExecuteIAMtoActor = false;
+				
+				// we can stop checking as soon as we know we have both cases in Selected Actors (the InActors array).
+				for (int32 Index = 0; Index < InActors.Num() && !(bCanExecuteActorToIAM && bCanExecuteIAMtoActor); ++Index)
+				{
+					const bool bIsIAM = InActors[0]->GetClass()->IsChildOf(AInstancedActorsManager::StaticClass());
+					// We can only convert and Actor to an IAM if it's not an IAM instance
+					bCanExecuteActorToIAM = bCanExecuteActorToIAM || !bIsIAM;
+					// We can only convert instances to and Actors only if it _is_ an IAM
+					bCanExecuteIAMtoActor = bCanExecuteIAMtoActor || bIsIAM;
+				}
+
+				MenuBuilder.AddMenuEntry(
+					FText::Format(LOCTEXT("ConvertSelectedActorsToIAMsText", "Convert {0} to Instanced Actors"), ActorName),
+					LOCTEXT("ConvertSelectedActorsToIAMsTooltip", "Convert the selected actors to Instanced Actors instances."),
+					FSlateIcon(FAppStyle::GetAppStyleSetName(), "Icons.Convert"),
+					FUIAction(
+						FExecuteAction::CreateRaw(this, &FInstancedActorsEditorModule::ConvertActorsToIAMsUIAction, InActors),
+						FCanExecuteAction::CreateLambda([bCanExecuteActorToIAM]() { return bCanExecuteActorToIAM; }))
+				);
+
+				MenuBuilder.AddMenuEntry(
+					FText::Format(LOCTEXT("ConvertSelectedIAMsToActorsText", "Convert {0}'s instances back to Actors"), ActorName),
+					LOCTEXT("ConvertSelectedIAMsToActorsToolTip", "Convert all the Instanced Actors instances back to Actors."),
+					FSlateIcon(FAppStyle::GetAppStyleSetName(), "Icons.Convert"),
+					FUIAction(
+						FExecuteAction::CreateRaw(this, &FInstancedActorsEditorModule::ConvertIAMsToActorsUIAction, InActors),
+						FCanExecuteAction::CreateLambda([bCanExecuteIAMtoActor]() { return bCanExecuteIAMtoActor; }))
+				);
+			})
+		);
+	}
+	return Extender;
+}
+
+void FInstancedActorsEditorModule::ConvertActorsToIAMsUIAction(const TArray<AActor*> InActors) const
+{
+	check(GEditor);
+
+	UWorld* World = GEditor->GetEditorWorldContext().World();
+	if (!World)
+	{
+		UE_LOG(LogInstancedActorsEditor, Log, TEXT("Unable to get Editor world."));
+		return;
+	}
+
+	UInstancedActorsSubsystem* IAMSubsystem = World->GetSubsystem<UInstancedActorsSubsystem>();
+	check(IAMSubsystem);
+
+	FScopedTransaction Transaction(LOCTEXT("ConvertToIAM_Transaction", "Convert Actors to IAMs"));
+	GEditor->SelectNone(/*bNoteSelectionChange=*/true, /*bDeselectBSPSurfs=*/true, /*WarnAboutManyActors=*/false);
+
+	for (AActor* Actor : InActors)
+	{
+		// note that we skip all the IAMs here, we don't support converting IAM to instances.
+		// We can end up here if there are multiple different AActors selected and some of them are IAMs - the option
+		// to convert will still appear in the "Actor" menu.
+		if (Actor == nullptr || InActors[0]->GetClass()->IsChildOf(AInstancedActorsManager::StaticClass()))
+		{
+			continue;
+		}
+
+		FInstancedActorsInstanceHandle InstanceHandle = IAMSubsystem->InstanceActor(Actor->GetClass(), Actor->GetActorTransform(), World->GetCurrentLevel());
+		if (InstanceHandle.IsValid())
+		{
+			Actor->Destroy(); // This will call modify too.
+			InstanceHandle.GetManager()->Modify();
+			GEditor->SelectActor(InstanceHandle.GetManager(), /*bInSelected*/ true, /*bNotify*/ true, /*bSelectEvenIfHidden*/ true);
+		}
+	}
+}
+
+void FInstancedActorsEditorModule::ConvertIAMsToActorsUIAction(const TArray<AActor*> InActors) const
+{
+	check(GEditor);
+
+	UWorld* World = GEditor->GetEditorWorldContext().World();
+	if (!World)
+	{
+		UE_LOG(LogInstancedActorsEditor, Log, TEXT("Unable to get Editor world."));
+		return;
+	}
+
+	UInstancedActorsSubsystem* IAMSubsystem = World->GetSubsystem<UInstancedActorsSubsystem>();
+	check(IAMSubsystem);
+
+	FScopedTransaction Transaction(LOCTEXT("ConvertToActorsFromIAM_Transaction", "Convert IAMs to Actors"));
+	GEditor->SelectNone(/*bNoteSelectionChange=*/true, /*bDeselectBSPSurfs=*/true, /*WarnAboutManyActors=*/false);
+
+	FActorSpawnParameters SpawnParams;
+	SpawnParams.SpawnCollisionHandlingOverride = ESpawnActorCollisionHandlingMethod::AlwaysSpawn;
+	SpawnParams.OverrideLevel = World->GetCurrentLevel();
+
+	for (AActor* Actor : InActors)
+	{
+		if (AInstancedActorsManager* Manager = Cast<AInstancedActorsManager>(Actor))
+		{
+			Manager->Modify();
+			FInstancedActorsIterationContext IterationContext;
+			Manager->ForEachInstance(
+				[&SpawnParams](const FInstancedActorsInstanceHandle& InstanceHandle, const FTransform& InstanceTransform, FInstancedActorsIterationContext& IterationContext)
+				{
+					AActor* SpawnedActor = nullptr;
+
+					const TSubclassOf<AActor>& ActorClass = InstanceHandle.GetInstanceActorDataChecked().ActorClass;
+					// Start by trying to use the ActorFactory.
+					if (UActorFactory* ActorFactory = GEditor->FindActorFactoryForActorClass(ActorClass))
+					{
+						SpawnedActor = ActorFactory->CreateActor(ActorClass, SpawnParams.OverrideLevel, InstanceTransform, SpawnParams);
+					}
+					// No factory found, fallback.
+					else 
+					{
+						SpawnedActor = GEditor->AddActor(SpawnParams.OverrideLevel, ActorClass, InstanceTransform, /*bSilent*/true);
+					}
+
+					if (SpawnedActor)
+					{
+						const FVector Scale3D = InstanceTransform.GetScale3D();
+						if (FVector::DistSquared(Scale3D, FVector::One()) > KINDA_SMALL_NUMBER)
+						{
+							SpawnedActor->SetActorScale3D(Scale3D);
+						}
+						GEditor->SelectActor(SpawnedActor, /*bInSelected*/ true, /*bNotify*/ true, /*bSelectEvenIfHidden*/ true);
+						InstanceHandle.GetManagerChecked().RemoveActorInstance(InstanceHandle);
+					}
+
+					// Continue
+					return true;
+				}
+				, IterationContext);
+		}
+	}
+}
+#undef LOCTEXT_NAMESPACE
diff --git a/Engine/Plugins/Runtime/InstancedActors/Source/InstancedActorsEditor/Public/InstancedActorsEditorModule.h b/Engine/Plugins/Runtime/InstancedActors/Source/InstancedActorsEditor/Public/InstancedActorsEditorModule.h
new file mode 100644
--- /dev/null
+++ b/Engine/Plugins/Runtime/InstancedActors/Source/InstancedActorsEditor/Public/InstancedActorsEditorModule.h
@@ -0,0 +1,36 @@
+// Copyright Epic Games, Inc. All Rights Reserved.
+
+#pragma  once
+
+#include "Containers/Array.h"
+#include "Templates/SharedPointer.h"
+#include "Delegates/IDelegateInstance.h"
+#include "Modules/ModuleInterface.h"
+
+class AActor;
+class FExtender;
+class FUICommandList;
+
+/**
+* The public interface to this module
+*/
+class INSTANCEDACTORSEDITOR_API FInstancedActorsEditorModule : public IModuleInterface
+{
+public:
+	// Begin IModuleInterface
+	virtual void StartupModule() override;
+	virtual void ShutdownModule() override;
+
+	/** Convert selected actors to Mass Instance Actors (IAMs). */
+	void ConvertActorsToIAMsUIAction(const TArray<AActor*> InActors) const;
+
+	/** Convert selected Mass Instance Actors to Actors. */
+	void ConvertIAMsToActorsUIAction(const TArray<AActor*> InActors) const;
+
+protected:
+	TSharedRef<FExtender> CreateLevelViewportContextMenuExtender(const TSharedRef<FUICommandList> CommandList, const TArray<AActor*> InActors);
+	void AddLevelViewportMenuExtender();
+	void RemoveLevelViewportMenuExtender();
+
+	FDelegateHandle LevelViewportExtenderHandle;
+};
