CIS Issue 516199: added MassEntity as an explicit dependency for InstancedActors (since it seems Mac doesn't like the way it's configured now)

diff --git a/Engine/Plugins/Runtime/InstancedActors/Source/InstancedActorsEditor/InstancedActorsEditor.Build.cs b/Engine/Plugins/Runtime/InstancedActors/Source/InstancedActorsEditor/InstancedActorsEditor.Build.cs
--- a/Engine/Plugins/Runtime/InstancedActors/Source/InstancedActorsEditor/InstancedActorsEditor.Build.cs
+++ b/Engine/Plugins/Runtime/InstancedActors/Source/InstancedActorsEditor/InstancedActorsEditor.Build.cs
@@ -24,6 +24,7 @@ namespace UnrealBuildTool.Rules
 					"Slate",
 					"LevelEditor",
 					"InstancedActors",
+					"MassEntity"
 				}
 			);
 		}
