CIS  Issue 516230: fixed wrong array index used in the lambda created by FInstancedActorsEditorModule::CreateLevelViewportContextMenuExtender's

diff --git a/Engine/Plugins/Runtime/InstancedActors/Source/InstancedActorsEditor/Private/InstancedActorsEditorModule.cpp b/Engine/Plugins/Runtime/InstancedActors/Source/InstancedActorsEditor/Private/InstancedActorsEditorModule.cpp
--- a/Engine/Plugins/Runtime/InstancedActors/Source/InstancedActorsEditor/Private/InstancedActorsEditorModule.cpp
+++ b/Engine/Plugins/Runtime/InstancedActors/Source/InstancedActorsEditor/Private/InstancedActorsEditorModule.cpp
@@ -77,7 +77,7 @@ TSharedRef<FExtender> FInstancedActorsEditorModule::CreateLevelViewportContextMe
 				// we can stop checking as soon as we know we have both cases in Selected Actors (the InActors array).
 				for (int32 Index = 0; Index < InActors.Num() && !(bCanExecuteActorToIAM && bCanExecuteIAMtoActor); ++Index)
 				{
-					const bool bIsIAM = InActors[0]->GetClass()->IsChildOf(AInstancedActorsManager::StaticClass());
+					const bool bIsIAM = InActors[Index]->GetClass()->IsChildOf(AInstancedActorsManager::StaticClass());
 					// We can only convert and Actor to an IAM if it's not an IAM instance
 					bCanExecuteActorToIAM = bCanExecuteActorToIAM || !bIsIAM;
 					// We can only convert instances to and Actors only if it _is_ an IAM
