Fix for townbuilder clear build site failing to detect instanced actors when you move away from it.

#rnx

diff --git a/Engine/Plugins/Runtime/InstancedActors/Source/InstancedActors/Private/InstancedActorsManager.cpp b/Engine/Plugins/Runtime/InstancedActors/Source/InstancedActors/Private/InstancedActorsManager.cpp
--- a/Engine/Plugins/Runtime/InstancedActors/Source/InstancedActors/Private/InstancedActorsManager.cpp
+++ b/Engine/Plugins/Runtime/InstancedActors/Source/InstancedActors/Private/InstancedActorsManager.cpp
@@ -1034,10 +1034,11 @@ bool AInstancedActorsManager::HasInstancesOfClass(const FBox& InQueryBounds, TSu
 	FScopedInstancedActorsIterationContext IterationContext;
 	bool bHasInstance = false;
 
-	auto InstancedActorDataMask = TOptional<AInstancedActorsManager::FInstancedActorDataPredicateFunc>([ActorClass, AllowedLODs](const UInstancedActorsData& InstancedActorData)
+	// @todo: disabling ActorClass testing to fix a game specific issue
+	auto InstancedActorDataMask = TOptional<AInstancedActorsManager::FInstancedActorDataPredicateFunc>([/*ActorClass, */AllowedLODs](const UInstancedActorsData& InstancedActorData)
 		{
-			return (int(AllowedLODs) & (1 << int(InstancedActorData.GetBulkLOD()))) 
-				&& InstancedActorData.ActorClass->IsChildOf(ActorClass);
+			return ((int(AllowedLODs) & (1 << int(InstancedActorData.GetBulkLOD()))) != 0);
+				/*&& InstancedActorData.ActorClass->IsChildOf(ActorClass)*/
 		});
 
 	struct FQueryBounds
