Cache UInstancedActorsData::AssetBounds in non-Editor scenarios as well, for runtime performance with AInstancedActorsManager::ForEachInstance() operations.

Most notably, the AActor::GetActorClassDefaultComponent<UStaticMeshComponent>() call within AInstancedActorsManager::CalculateBounds() can be expensive

diff --git a/Engine/Plugins/Runtime/InstancedActors/Source/InstancedActors/Private/InstancedActorsData.cpp b/Engine/Plugins/Runtime/InstancedActors/Source/InstancedActors/Private/InstancedActorsData.cpp
--- a/Engine/Plugins/Runtime/InstancedActors/Source/InstancedActors/Private/InstancedActorsData.cpp
+++ b/Engine/Plugins/Runtime/InstancedActors/Source/InstancedActors/Private/InstancedActorsData.cpp
@@ -182,6 +182,8 @@ void UInstancedActorsData::Initialize()
 	if (Settings && Settings->bOverride_ActorClass && Settings->ActorClass)
 	{
 		ActorClass = Settings->ActorClass;
+		CachedLocalBounds = AInstancedActorsManager::CalculateBounds(ActorClass);
+		ensure(CachedLocalBounds.IsValid);
 	}
 
 	// Allow settings to scale NumValidInstances, effectively scaling the number of spawned entities
@@ -558,8 +560,8 @@ FInstancedActorsInstanceHandle UInstancedActorsData::AddInstance(const FTransfor
 		InstanceTransforms[NewInstanceIndex].SetToRelativeTransform(ManagerTransform);
 	}
 
-	ensure(AssetBounds.IsValid);
-	Bounds += AssetBounds.TransformBy(InstanceTransforms[NewInstanceIndex]);
+	ensure(CachedLocalBounds.IsValid);
+	Bounds += CachedLocalBounds.TransformBy(InstanceTransforms[NewInstanceIndex]);
 
 	return FInstancedActorsInstanceHandle(*this, FInstancedActorsInstanceIndex(NewInstanceIndex));
 }
@@ -581,13 +583,13 @@ bool UInstancedActorsData::RemoveInstance(const FInstancedActorsInstanceHandle&
 		}
 
 		// Update bounds
-		ensure(AssetBounds.IsValid);
+		ensure(CachedLocalBounds.IsValid);
 		Bounds.Init();
 		for (const FTransform& InstanceTransform : InstanceTransforms)
 		{
 			if (UE::InstancedActors::Helpers::IsValidInstanceTransform(InstanceTransform))
 			{
-				Bounds += AssetBounds.TransformBy(InstanceTransform);
+				Bounds += CachedLocalBounds.TransformBy(InstanceTransform);
 			}
 		}
 
@@ -620,13 +622,13 @@ bool UInstancedActorsData::SetInstanceTransform(const FInstancedActorsInstanceHa
 	}
 
 	// Update bounds
-	ensure(AssetBounds.IsValid);
+	ensure(CachedLocalBounds.IsValid);
 	Bounds.Init();
 	for (const FTransform& InstanceTransform : InstanceTransforms)
 	{
 		if (UE::InstancedActors::Helpers::IsValidInstanceTransform(InstanceTransform))
 		{
-			Bounds += AssetBounds.TransformBy(InstanceTransform);
+			Bounds += CachedLocalBounds.TransformBy(InstanceTransform);
 		}
 	}
 
@@ -706,27 +708,31 @@ void UInstancedActorsData::PostLoad()
 
 	NumInstances = InstanceTransforms.Num();
 
+	// Cache asset bounds
+	if (!CachedLocalBounds.IsValid)
+	{
+		CachedLocalBounds = AInstancedActorsManager::CalculateBounds(ActorClass);
+	}
+	ensure(CachedLocalBounds.IsValid);
+	
 	if (!Bounds.IsValid)
 	{
 		// Update bounds for InstanceData saved prior to addition of Bounds property
 		if (NumValidInstances > 0)
 		{
-			FBox MeshBounds = AInstancedActorsManager::CalculateBounds(ActorClass);
-			ensure(MeshBounds.IsValid);
-
 			Bounds.Init();
 			for (const FTransform& InstanceTransform : InstanceTransforms)
 			{
 				if (UE::InstancedActors::Helpers::IsValidInstanceTransform(InstanceTransform))
 				{
-					Bounds += MeshBounds.TransformBy(InstanceTransform);
+					Bounds += CachedLocalBounds.TransformBy(InstanceTransform);
 				}
 			}
 		}
 		else
 		{
 			// Note: UInstancedActorsData::Bounds mustn't be 0 sized by default, otherwise the first AddInstance
-			//       woulf stretch the bounds from origin - instead we want the first AddInstance to init the bounds.
+			//       would stretch the bounds from origin - instead we want the first AddInstance to init the bounds.
 			Bounds = FBox(FVector::ZeroVector, FVector::ZeroVector);
 		}
 	}
@@ -746,14 +752,6 @@ void UInstancedActorsData::PostLoad()
 		}
 		EditorPreviewISMComponents.Empty();
 	}
-
-	// Cache asset bounds for use during instance population for IAD's saved before
-	// AssetBounds was marked non-transient.
-	if (!AssetBounds.IsValid)
-	{
-		AssetBounds = AInstancedActorsManager::CalculateBounds(ActorClass);
-	}
-	ensure(AssetBounds.IsValid);
 #endif
 
 	//check(CompressedInstanceTransforms.IsEmpty());
diff --git a/Engine/Plugins/Runtime/InstancedActors/Source/InstancedActors/Private/InstancedActorsManager.cpp b/Engine/Plugins/Runtime/InstancedActors/Source/InstancedActors/Private/InstancedActorsManager.cpp
--- a/Engine/Plugins/Runtime/InstancedActors/Source/InstancedActors/Private/InstancedActorsManager.cpp
+++ b/Engine/Plugins/Runtime/InstancedActors/Source/InstancedActors/Private/InstancedActorsManager.cpp
@@ -704,8 +704,8 @@ UInstancedActorsData& AInstancedActorsManager::GetOrCreateActorInstanceData(TSub
 	RegisterInstanceDatasComponents(*NewInstanceData, NewInstanceData->EditorPreviewISMComponents);
 
 	// Cache asset bounds for use during instance population
-	NewInstanceData->AssetBounds = CalculateBounds(ActorClass);
-	ensure(NewInstanceData->AssetBounds.IsValid);
+	NewInstanceData->CachedLocalBounds = CalculateBounds(ActorClass);
+	ensure(NewInstanceData->CachedLocalBounds.IsValid);
 
 	return *NewInstanceData;
 }
@@ -948,7 +948,7 @@ bool AInstancedActorsManager::ForEachInstance<FSphere>(const FSphere& QueryBound
 			}
 
 			// More expensive bounds test.
-			const FBox InstancedActorBounds = CalculateBounds(InstanceHandle.InstancedActorData->ActorClass);
+			const FBox& InstancedActorBounds = InstanceHandle.InstancedActorData->CachedLocalBounds;
 			const FSphere TransformedSphere = QueryBounds.TransformBy(InstanceTransform.Inverse());
 			if (FMath::SphereAABBIntersection(TransformedSphere, InstancedActorBounds))
 			{
@@ -971,7 +971,7 @@ bool AInstancedActorsManager::ForEachInstance<FBox>(const FBox& QueryBounds, FIn
 			}
 
 			// More expensive bounds test.
-			const FBox InstancedActorBounds = CalculateBounds(InstanceHandle.InstancedActorData->ActorClass).TransformBy(InstanceTransform);
+			const FBox InstancedActorBounds = InstanceHandle.InstancedActorData->CachedLocalBounds.TransformBy(InstanceTransform);
 			if (QueryBounds.Intersect(InstancedActorBounds))
 			{
 				return Operation(InstanceHandle, InstanceTransform, IterationContext);
@@ -995,7 +995,7 @@ UE::InstancedActors::EInsideBoundsTestResult AInstancedActorsManager::IsInstance
 	}
 
 	// More expensive bounds test.
-	const FBox InstancedActorBounds = CalculateBounds(InstanceHandle.InstancedActorData->ActorClass).TransformBy(InstanceTransform);
+	const FBox InstancedActorBounds = InstanceHandle.InstancedActorData->CachedLocalBounds.TransformBy(InstanceTransform);
 	if (QueryBounds.Intersect(InstancedActorBounds))
 	{
 		return UE::InstancedActors::EInsideBoundsTestResult::OverlapBounds;
@@ -1013,7 +1013,7 @@ UE::InstancedActors::EInsideBoundsTestResult AInstancedActorsManager::IsInstance
 	}
 
 	// More expensive bounds test.
-	const FBox InstancedActorBounds = CalculateBounds(InstanceHandle.InstancedActorData->ActorClass);
+	const FBox& InstancedActorBounds = InstanceHandle.InstancedActorData->CachedLocalBounds;
 	const FSphere TransformedSphere = QueryBounds.TransformBy(InstanceTransform.Inverse());
 	if (FMath::SphereAABBIntersection(TransformedSphere, InstancedActorBounds))
 	{
@@ -1096,7 +1096,7 @@ bool AInstancedActorsManager::HasInstancesOfClass(const FBox& InQueryBounds, TSu
 			(
 				if (bHasInstance || OverlapResult == EInsideBoundsTestResult::OverlapBounds)
 				{
-					const FBox InstancedActorBounds = CalculateBounds(InstanceHandle.InstancedActorData->ActorClass).TransformBy(InstanceTransform);
+					const FBox InstancedActorBounds = InstanceHandle.InstancedActorData->CachedLocalBounds.TransformBy(InstanceTransform);
 					UE_VLOG_BOX(Manager, LogInstancedActors, Log, InstancedActorBounds, bHasInstance ? FColor::Red : FColor::Green
 						, TEXT("Instance of class %s"), *GetNameSafe(InstanceHandle.InstancedActorData->ActorClass));
 				}
diff --git a/Engine/Plugins/Runtime/InstancedActors/Source/InstancedActors/Public/InstancedActorsData.h b/Engine/Plugins/Runtime/InstancedActors/Source/InstancedActors/Public/InstancedActorsData.h
--- a/Engine/Plugins/Runtime/InstancedActors/Source/InstancedActors/Public/InstancedActorsData.h
+++ b/Engine/Plugins/Runtime/InstancedActors/Source/InstancedActors/Public/InstancedActorsData.h
@@ -347,6 +347,8 @@ public:
 
 	int32 GetInstanceDataID() const { return (int32)ID; }
 
+	const FBox& GetCachedLocalBounds() const { return CachedLocalBounds; }
+
 	FMassEntityHandle GetEntityHandleForIndex(const FInstancedActorsInstanceIndex Index) const
 	{
 		return Entities.IsValidIndex(Index.GetIndex()) ? Entities[Index.GetIndex()] : FMassEntityHandle();
@@ -430,12 +432,10 @@ protected:
 	// ISMCs created in GetOrCreateActorInstanceData to match default visualizations ISMComponents for editor only preview of instances
 	UPROPERTY()
 	TArray<TObjectPtr<UInstancedStaticMeshComponent>> EditorPreviewISMComponents;
+#endif
 
-	// Editor time cached bounds for AssetClass, used during editor-only instance creation and deletion to
-	// expand / contract Bounds
 	UPROPERTY()
-	FBox AssetBounds = FBox(ForceInit);
-#endif
+	FBox CachedLocalBounds = FBox(ForceInit);
 
 private:
 	// Represents the shared fragment registered with MassEntityManager, that points back to this UInstancedActorsData instance
