Attempting to fix optimized builds' AInstancedActorsManager::HasInstancesOfClass issue failing to test IAData's actor class with a provided filer lambda. Changed the lambda's capture mode to "by value" in hopes that it was the culprit of difference between optimized and non-optimized builds.

diff --git a/Engine/Plugins/Runtime/InstancedActors/Source/InstancedActors/Private/InstancedActorsManager.cpp b/Engine/Plugins/Runtime/InstancedActors/Source/InstancedActors/Private/InstancedActorsManager.cpp
--- a/Engine/Plugins/Runtime/InstancedActors/Source/InstancedActors/Private/InstancedActorsManager.cpp
+++ b/Engine/Plugins/Runtime/InstancedActors/Source/InstancedActors/Private/InstancedActorsManager.cpp
@@ -1034,11 +1034,10 @@ bool AInstancedActorsManager::HasInstancesOfClass(const FBox& InQueryBounds, TSu
 	FScopedInstancedActorsIterationContext IterationContext;
 	bool bHasInstance = false;
 
-	// @todo: disabling ActorClass testing to fix a game specific issue
-	auto InstancedActorDataMask = TOptional<AInstancedActorsManager::FInstancedActorDataPredicateFunc>([/*ActorClass, */AllowedLODs](const UInstancedActorsData& InstancedActorData)
+	auto InstancedActorDataMask = TOptional<AInstancedActorsManager::FInstancedActorDataPredicateFunc>([=](const UInstancedActorsData& InstancedActorData)
 		{
-			return ((int(AllowedLODs) & (1 << int(InstancedActorData.GetBulkLOD()))) != 0);
-				/*&& InstancedActorData.ActorClass->IsChildOf(ActorClass)*/
+			return ((int(AllowedLODs) & (1 << int(InstancedActorData.GetBulkLOD()))) != 0)
+				&& InstancedActorData.ActorClass->IsChildOf(ActorClass);
 		});
 
 	struct FQueryBounds
