[InstancedActors] Modified IAM registration a bit to ensure we don't get double-registration attempts for instanced that are placed in the persistent level

The CL also contains a minor clean up of inline functions and comments in InstancedActorsManager.h

diff --git a/Engine/Plugins/Runtime/InstancedActors/Source/InstancedActors/Private/InstancedActorsManager.cpp b/Engine/Plugins/Runtime/InstancedActors/Source/InstancedActors/Private/InstancedActorsManager.cpp
--- a/Engine/Plugins/Runtime/InstancedActors/Source/InstancedActors/Private/InstancedActorsManager.cpp
+++ b/Engine/Plugins/Runtime/InstancedActors/Source/InstancedActors/Private/InstancedActorsManager.cpp
@@ -191,6 +191,10 @@ void AInstancedActorsManager::BeginPlay()
 	}
 	MassEntityManager = EntitySubsystem->GetMutableEntityManager().AsShared();
 
+	// the main reason for not supporting pre-BeginPlay registration is that we cache the InstanceBounds here and the
+	// bounds are used during manager's registraiton to place it on the 2d hashmap
+	checkf(ManagerHandle.IsValid() == false, TEXT("We don't expect IAMs to be registered before their BeginPlay."));
+
 	// Cache world instance bounds
 	//
 	// Note: This must be done before InstancedActorSubsystem->AddManager bewlow, as this is used to
diff --git a/Engine/Plugins/Runtime/InstancedActors/Source/InstancedActors/Private/InstancedActorsSubsystem.cpp b/Engine/Plugins/Runtime/InstancedActors/Source/InstancedActors/Private/InstancedActorsSubsystem.cpp
--- a/Engine/Plugins/Runtime/InstancedActors/Source/InstancedActors/Private/InstancedActorsSubsystem.cpp
+++ b/Engine/Plugins/Runtime/InstancedActors/Source/InstancedActors/Private/InstancedActorsSubsystem.cpp
@@ -115,19 +115,26 @@ UInstancedActorsSubsystem& UInstancedActorsSubsystem::GetChecked(UObject* WorldC
 
 bool UInstancedActorsSubsystem::ShouldCreateSubsystem(UObject* Outer) const
 {
-	if(!Super::ShouldCreateSubsystem(Outer))
+	if (Super::ShouldCreateSubsystem(Outer) && Outer)
 	{
-		return false;
-	}
+		if (UWorld* World = Outer->GetWorld())
+		{
+#if WITH_EDITOR
+			// we don't want to create subsystems for Editor worlds while PIE is active
+			// This wouldn't happen in normal world lifecycle, but can happen if FSubsystemCollectionBase::ActivateExternalSubsystem
+			// is used (it adds an instance of a given subsystem class to ALL worlds) - for example by GameFeatureActions.
+			if (GEditor && GEditor->IsPlayingSessionInEditor() && World->WorldType == EWorldType::Editor)
+			{
+				return false;
+			}
+#endif // WITH_EDITOR
 
-	// UInstancedActorsSubsystem must always be present for editor worlds to allow for InstanceActor etc editor operations
-	UWorld* World = CastChecked<UWorld>(Outer);
-	if (World->WorldType == EWorldType::Editor)
-	{
-		return true;
+			// we only ever want to have a single instance of this subsystem. Attempting to add multiple
+			// instances can be a result of subsystem adding game feature actions.
+			return World->GetSubsystemBase(GetClass()) == nullptr;
+		}
 	}
-
-	return true;
+	return false;
 }
 
 void UInstancedActorsSubsystem::Initialize(FSubsystemCollectionBase& Collection)
@@ -175,12 +182,15 @@ void UInstancedActorsSubsystem::Initialize(FSubsystemCollectionBase& Collection)
 	}
 
 	// Collect existing managers, calling AInstancedActorsManager::OnAddedToSubsystem to inform them of latent addition to this subsystem
-	for (TActorIterator<AInstancedActorsManager> MangerIt(World); MangerIt; ++MangerIt)
+	for (TActorIterator<AInstancedActorsManager> ManagerIt(World); ManagerIt; ++ManagerIt)
 	{
-		AInstancedActorsManager* Manager = *MangerIt;
+		AInstancedActorsManager* Manager = *ManagerIt;
 		check(Manager);
-
-		AddManager(*Manager);
+		// we only case about managers that have already begun play and missed their chance to registred in their BeginPlay
+		if (Manager->HasActorBegunPlay())
+		{
+			AddManager(*Manager);
+		}
 	}
 }
 
@@ -216,24 +226,30 @@ TStatId UInstancedActorsSubsystem::GetStatId() const
 
 FInstancedActorsManagerHandle UInstancedActorsSubsystem::AddManager(AInstancedActorsManager& Manager)
 {
+	FInstancedActorsManagerHandle ManagerHandle;
 	const FBox ManagerBounds = Manager.GetInstanceBounds();
 
-	check(Algo::Find(Managers, &Manager) == nullptr);
-	const int32 ManagerID = Managers.Add(&Manager);
-
-	FInstancedActorsManagerHandle ManagerHandle = ManagerID;
-	ManagersHashGrid.Add(ManagerHandle, ManagerBounds);
+	if (ensureMsgf(Algo::Find(Managers, &Manager) == nullptr, TEXT("A given Manager instance is not expected to be added twice")))
+	{
+		ManagerHandle = Managers.Add(&Manager);
+		ManagersHashGrid.Add(ManagerHandle, ManagerBounds);
 
 #if WITH_INSTANCEDACTORS_DEBUG
-	// Record initial bounds so we can compare on removal to make sure it wasn't changed
-	DebugManagerBounds.Add(&Manager, ManagerBounds);
+		// Record initial bounds so we can compare on removal to make sure it wasn't changed
+		DebugManagerBounds.Add(&Manager, ManagerBounds);
 #endif
 
-	// Let Manager know the subsystem is ready. 
-	//
-	// Common callback for both AInstancedActorsManager::BeginPlay -> AddManager and latent 
-	// UInstancedActorsSubsystem::Initialize -> AddManager
-	Manager.OnAddedToSubsystem(*this, ManagerHandle);
+		// Let Manager know the subsystem is ready. 
+		//
+		// Common callback for both AInstancedActorsManager::BeginPlay -> AddManager and latent 
+		// UInstancedActorsSubsystem::Initialize -> AddManager
+		Manager.OnAddedToSubsystem(*this, ManagerHandle);
+	}
+	else
+	{
+		ManagerHandle = Manager.GetManagerHandle();
+		checkf(ManagerHandle.IsValid(), TEXT("If a given Manager has already been registered we expect it to host a valid ManagerHandle"));
+	}
 	
 	return ManagerHandle;
 }
diff --git a/Engine/Plugins/Runtime/InstancedActors/Source/InstancedActors/Public/InstancedActorsManager.h b/Engine/Plugins/Runtime/InstancedActors/Source/InstancedActors/Public/InstancedActorsManager.h
--- a/Engine/Plugins/Runtime/InstancedActors/Source/InstancedActors/Public/InstancedActorsManager.h
+++ b/Engine/Plugins/Runtime/InstancedActors/Source/InstancedActors/Public/InstancedActorsManager.h
@@ -60,21 +60,26 @@ class INSTANCEDACTORS_API AInstancedActorsManager : public APartitionActor, publ
 public:
 	AInstancedActorsManager();
 
-	// Adds modifiers already registered with InInstancedActorSubsystem and either calls InitializeModifyAndSpawnEntities to spawn
-	// entities immediately, or schedules deferred call by InInstancedActorSubsystem if IA.DeferSpawnEntities is enabled.
-	//
-	// Called either in BeginPlay if InInstancedActorSubsystem was already initialized or latently once it is, in
-	// UInstancedActorsSubsystem::Initialize
+	/** 
+	 * Adds modifiers already registered with InInstancedActorSubsystem and either calls InitializeModifyAndSpawnEntities to spawn
+	 * entities immediately, or schedules deferred call by InInstancedActorSubsystem if IA.DeferSpawnEntities is enabled.
+	 * Called either in BeginPlay if InInstancedActorSubsystem was already initialized or latently once it is, in
+	 * UInstancedActorsSubsystem::Initialize
+	 */
 	void OnAddedToSubsystem(UInstancedActorsSubsystem& InInstancedActorSubsystem, FInstancedActorsManagerHandle InManagerHandle);
 
-	// Initializes all PerActorClassInstanceData, applies pre-spawn modifiers, spawns entities then applies post-spawn modifiers.
-	//
-	// Called either directly in OnAddedToSubsystem or deferred and time-sliced in UInstancedActorsSubsystem::Tick if
-	// IA.DeferSpawnEntities is enabled.
+	FInstancedActorsManagerHandle GetManagerHandle() const;
+
+	/** 
+	 * Initializes all PerActorClassInstanceData, applies pre-spawn modifiers, spawns entities then applies post-spawn modifiers.
+	 * 
+	 * Called either directly in OnAddedToSubsystem or deferred and time-sliced in UInstancedActorsSubsystem::Tick if
+	 * IA.DeferSpawnEntities is enabled.
+	 */
 	void InitializeModifyAndSpawnEntities();
 
-	// Returns true if InstanceTransforms have been consumed to spawn Mass entities in InitializeModifyAndSpawnEntities
-	bool HasSpawnedEntities() const { return bHasSpawnedEntities; }
+	/** @return true if InstanceTransforms have been consumed to spawn Mass entities in InitializeModifyAndSpawnEntities */
+	bool HasSpawnedEntities() const;
 
 #if WITH_EDITOR
 	/** Adds an instance of ActorClass at InstanceTransform location to instance data */
@@ -94,11 +99,11 @@ public:
 	bool RemoveActorInstance(const FInstancedActorsInstanceHandle& InstanceToRemove);
 #endif
 
-	// Searches PerActorClassInstanceData, returning the IAD with matching UInstancedActorsData::ID, if any (nullptr otherwise)
+	/** Searches PerActorClassInstanceData, returning the IAD with matching UInstancedActorsData::ID, if any(nullptr otherwise) */
 	UInstancedActorsData* FindInstanceDataByID(uint16 InstanceDataID) const;
 
-	// Returns the full set of instance data for this manager
-	TConstArrayView<TObjectPtr<UInstancedActorsData>> GetAllInstanceData() const { return PerActorClassInstanceData; }
+	/** @return the full set of instance data for this manager */
+	TConstArrayView<TObjectPtr<UInstancedActorsData>> GetAllInstanceData() const;
 
 	/**
 	 * Removes all instances as if they were never present i.e: these removals are not persisted as
@@ -110,17 +115,19 @@ public:
 	 */
 	void RuntimeRemoveAllInstances();
 
-	// Returns the current valid instance count (i.e: NumInstances - FreeList.Num()) sum for all instance datas
+	/** @return the current valid instance count (i.e: NumInstances - FreeList.Num()) sum for all instance datas */
 	int32 GetNumValidInstances() const;
 
 	bool HasAnyValidInstances() const;
 
-	// Returns true if InstanceHandle refers to this manager and we have current information for an
-	// instance at InstanceHandle.InstanceIndex in InstanceHandle.InstancedActorData
+	/**
+	 * @return true if InstanceHandle refers to this manager and we have current information for an
+	 *	instance at InstanceHandle.InstanceIndex in InstanceHandle.InstancedActorData  
+	 */
 	bool IsValidInstance(const FInstancedActorsInstanceHandle& InstanceHandle) const;
 
-	// Returns world space cumulative instance bounds. Only valid after BeginPlay.
-	FBox GetInstanceBounds() const { return InstanceBounds; }
+	/** @return world space cumulative instance bounds. Only valid after BeginPlay. */
+	FBox GetInstanceBounds() const;
 
 	/**
 	 * Iteration callback for ForEachInstance
@@ -180,37 +187,29 @@ public:
 	static UE::InstancedActors::EInsideBoundsTestResult IsInstanceInsideBounds(const TBoundsType & QueryBounds
 		, const FInstancedActorsInstanceHandle& InstanceHandle, const FTransform & InstanceTransform);
 
-	// Outputs instance metrics to Ar
+	/** Outputs instance metrics to Ar */
 	void AuditInstances(FOutputDevice& Ar, bool bDebugDraw = false, float DebugDrawDuration = 10.0f) const;
 
-	// Called by IA.CompactInstances console command to fully remove FreeList instances
+	/** Called by IA.CompactInstances console command to fully remove FreeList instances */
 	void CompactInstances(FOutputDevice& Ar);
 
 	void AddModifierVolume(UInstancedActorsModifierVolumeComponent& ModifierVolume);
 	void RemoveModifierVolume(UInstancedActorsModifierVolumeComponent& ModifierVolume);
 	void RemoveAllModifierVolumes();
 
-	// Request the persistant data system to re-save this managers persistent data
+	/** Request the persistent data system to re-save this managers persistent data */
 	void RequestPersistentDataSave();
 
-	// Helper function to deduce appropriate instanced static mesh bounds for ActorClass
+	/** Helper function to deduce appropriate instanced static mesh bounds for ActorClass */
 	static FBox CalculateBounds(TSubclassOf<AActor> ActorClass);
 
-	// Returns the Mass entity manager used to spawn entities. Valid only after BeginPlay
-	FORCEINLINE TSharedPtr<FMassEntityManager> GetMassEntityManager() const { return MassEntityManager; }
-	FORCEINLINE FMassEntityManager& GetMassEntityManagerChecked() const
-	{
-		check(MassEntityManager.IsValid());
-		return *MassEntityManager;
-	}
+	/** @return the Mass entity manager used to spawn entities. Valid only after BeginPlay */
+	TSharedPtr<FMassEntityManager> GetMassEntityManager() const;
+	FMassEntityManager& GetMassEntityManagerChecked() const;
 
-	// Returns the Instanced Actor Subsystem this Manager is registered with. Valid only after BeginPlay
-	FORCEINLINE UInstancedActorsSubsystem* GetInstancedActorSubsystem() const { return InstancedActorSubsystem; }
-	FORCEINLINE UInstancedActorsSubsystem& GetInstancedActorSubsystemChecked() const
-	{
-		check(InstancedActorSubsystem);
-		return *InstancedActorSubsystem;
-	}
+	/** @return the Instanced Actor Subsystem this Manager is registered with. Valid only after BeginPlay */
+	UInstancedActorsSubsystem* GetInstancedActorSubsystem() const;
+	UInstancedActorsSubsystem& GetInstancedActorSubsystemChecked() const;
 
 	//~ Begin APartitionActor Overrides
 #if WITH_EDITOR
@@ -267,34 +266,37 @@ protected:
 	virtual FTransform GetTransform(const FActorInstanceHandle& Handle) const override;
 	//~ End IActorInstanceManagerInterface Overrides
 
-	// Called by Serialize for SaveGame archives to save / load IAD persistence data
-	// @param Record		The archive record to read / write IAD save data to
-	// @param InstanceData	The InstanceData to serialize from / to. May be nullptr when loading if Record's IAD has been removed since saving.
-	//						In this case, we still need to read Record to seek the archive past this IAD record consistently.
-	// @param TimeDelta		Real time in seconds since serialization (0 when saving)
+	/** 
+	 * Called by Serialize for SaveGame archives to save / load IAD persistence data
+	 * @param Record		The archive record to read / write IAD save data to
+	 * @param InstanceData	The InstanceData to serialize from / to. May be nullptr when loading if Record's IAD has been removed since saving.
+	 * 						In this case, we still need to read Record to seek the archive past this IAD record consistently.
+	 * @param TimeDelta		Real time in seconds since serialization (0 when saving)
+	 */
 	void SerializeInstancePersistenceData(FStructuredArchive::FRecord Record, UInstancedActorsData* InstanceData, int64 TimeDelta) const;
 
-	// Despawns all entities spawned by individual UInstancedActorsData instances. 
+	/** Despawns all entities spawned by individual UInstancedActorsData instances. */
 	virtual void DespawnAllEntities();
 
-	// Attempts to run any 'pending' modifiers in ModifierVolumes where are appropriate to run givem HasSpawnedEntities
-	// Called in BeginPlay prior to, and then again after SpawnEntities. Also called in AddModifierVolume.
-	// @see UInstancedActorsModifierBase::bRequiresSpawnedEntities
+	/** Attempts to run any 'pending' modifiers in ModifierVolumes where are appropriate to run given HasSpawnedEntities
+	 * Called in BeginPlay prior to, and then again after SpawnEntities. Also called in AddModifierVolume.
+	 * @see UInstancedActorsModifierBase::bRequiresSpawnedEntities
+	 */
 	void TryRunPendingModifiers();
 
-	// Called when persistent data has been applied / restored
+	/** Called when persistent data has been applied / restored */
 	void OnPersistentDataRestored();
 
-	// Calculate cumulative local space instance bounds for all PerActorClassInstanceData
+	/** Calculate cumulative local space instance bounds for all PerActorClassInstanceData */
 	FBox CalculateLocalInstanceBounds() const;
 
 #if WITH_EDITOR
-	// Helper function to create and initialize per-actor-class UInstancedActorsData's, optionally further partitioned by InstanceTags
+	/** Helper function to create and initialize per-actor-class UInstancedActorsData's, optionally further partitioned by InstanceTags */
 	UInstancedActorsData& GetOrCreateActorInstanceData(TSubclassOf<AActor> ActorClass, const FInstancedActorsTagSet& InstanceTags);
 
 	virtual UInstancedActorsData* CreateNextInstanceActorData(TSubclassOf<AActor> ActorClass, const FInstancedActorsTagSet& InstanceTags);
 
-	// Used to set the right properties on the editor ISMCs so we can do per-instance selection.
+	/** Used to set the right properties on the editor ISMCs so we can do per-instance selection. */
 	virtual void PreRegisterAllComponents() override;
 #endif
 
@@ -306,7 +308,7 @@ protected:
 
 	TSharedPtr<FMassEntityManager> MassEntityManager;
 
-	/** Saved Actor Guid. Initialized from the actor name in constructor */
+	/** Saved Actor Guid. Initialized from the actor name in constructor */ 
 	UPROPERTY(VisibleAnywhere, Category=InstancedActors)
 	FGuid SavedActorGuid = FGuid();
 
@@ -314,34 +316,40 @@ protected:
 	UPROPERTY(Transient)
 	bool bHasSpawnedEntities = false;
 
-	// Incremented in GetOrCreateActorInstanceData to provide IAD's with a stable, unique identifier
-	// within this IAM.
-	// @see UInstancedActorsData::ID
+	/** 
+	 * Incremented in GetOrCreateActorInstanceData to provide IAD's with a stable, unique identifier
+	 * within this IAM.
+	 * @see UInstancedActorsData::ID
+	 */
 	UPROPERTY()
 	uint16 NextInstanceDataID = 0;
 
-	// Per-actor-class instance data populated by AddActorInstance
+	/** Per-actor-class instance data populated by AddActorInstance */
 	UPROPERTY(Instanced, VisibleAnywhere, Category=InstancedActors)
 	TArray<TObjectPtr<UInstancedActorsData>> PerActorClassInstanceData;
 
-	// World space cumulative instance bounds, calculated in BeginPlay
+	/** World space cumulative instance bounds, calculated in BeginPlay */
 	UPROPERTY(Transient)
 	FBox InstanceBounds = FBox(ForceInit);
 
-	// Modifier volumes added via AddModifierVolume
+	/** Modifier volumes added via AddModifierVolume */
 	TArray<TWeakObjectPtr<UInstancedActorsModifierVolumeComponent>> ModifierVolumes;
 
-	// A bit flag per volume in ModifierVolumes for whether the volume has pending Modifiers to run on this manager
-	// i.e if PendingModifierVolumeModifiers[VolumeIndex] has any true flags.
-	// Some modifiers can run prior to entity spawning for efficiency purposes, others must be held 'pending'
-	// later execution after entities have spawned.
+	/** 
+	 * A bit flag per volume in ModifierVolumes for whether the volume has pending Modifiers to run on this manager
+	 * i.e if PendingModifierVolumeModifiers[VolumeIndex] has any true flags.
+	 * Some modifiers can run prior to entity spawning for efficiency purposes, others must be held 'pending'
+	 * later execution after entities have spawned.
+	 */
 	TBitArray<> PendingModifierVolumes;
 
-	// A set of bit flags per volume in ModifierVolumes, matching each modifiers Modifiers list, marking
-	// whether the Modifier has yet to run on this manager or not (true = needs running).
-	//
-	// Some modifiers can run prior to entity spawning for efficiency purposes, others must be held 'pending'
-	// later execution after entities have spawned.
+	/** 
+	 * A set of bit flags per volume in ModifierVolumes, matching each modifiers Modifiers list, marking
+	 * whether the Modifier has yet to run on this manager or not (true = needs running).
+	 * 
+	 * Some modifiers can run prior to entity spawning for efficiency purposes, others must be held 'pending'
+	 * later execution after entities have spawned.
+	 */
 	TArray<TBitArray<>> PendingModifierVolumeModifiers;
 
 	UPROPERTY(Transient)
@@ -358,7 +366,7 @@ private:
 	UPROPERTY()
 	FGuid ManagerGridGuid;
 
-	// Set this to false to be able to move the instances contained by this IAM. The property is not saved and will reset.
+	/** Set this to false to be able to move the instances contained by this IAM.The property is not saved and will reset. */
 	UPROPERTY(EditAnywhere, Transient, meta = (DisplayPriority = 1), Category = InstancedActors)
 	bool bLockInstanceLocation = true;
 #endif
@@ -388,3 +396,49 @@ private:
 
 	FInstancedActorsInstanceHandle ActorInstanceHandleFromFSMInstanceId(const FSMInstanceId& InstanceId) const;
 };
+
+
+//-----------------------------------------------------------------------------
+// inlines
+//-----------------------------------------------------------------------------
+FORCEINLINE FInstancedActorsManagerHandle AInstancedActorsManager::GetManagerHandle() const
+{ 
+	return ManagerHandle; 
+}
+
+FORCEINLINE FBox AInstancedActorsManager::GetInstanceBounds() const
+{ 
+	return InstanceBounds; 
+}
+
+FORCEINLINE bool AInstancedActorsManager::HasSpawnedEntities() const
+{ 
+	return bHasSpawnedEntities; 
+}
+
+FORCEINLINE TConstArrayView<TObjectPtr<UInstancedActorsData>> AInstancedActorsManager::GetAllInstanceData() const
+{ 
+	return PerActorClassInstanceData; 
+}
+
+FORCEINLINE TSharedPtr<FMassEntityManager> AInstancedActorsManager::GetMassEntityManager() const 
+{ 
+	return MassEntityManager; 
+}
+
+FORCEINLINE FMassEntityManager& AInstancedActorsManager::GetMassEntityManagerChecked() const
+{
+	check(MassEntityManager.IsValid());
+	return *MassEntityManager;
+}
+
+FORCEINLINE UInstancedActorsSubsystem* AInstancedActorsManager::GetInstancedActorSubsystem() const 
+{ 
+	return InstancedActorSubsystem; 
+}
+
+FORCEINLINE UInstancedActorsSubsystem& AInstancedActorsManager::GetInstancedActorSubsystemChecked() const
+{
+	check(InstancedActorSubsystem);
+	return *InstancedActorSubsystem;
+}
