Marked up InstancedActorsVisualizationTrait to now show up in class picker dropdown, since the class is not intended to be used in entity configs. Class's comment has been updated to reflect that.

diff --git a/Engine/Plugins/Runtime/InstancedActors/Source/InstancedActors/Public/InstancedActorsVisualizationTrait.h b/Engine/Plugins/Runtime/InstancedActors/Source/InstancedActors/Public/InstancedActorsVisualizationTrait.h
--- a/Engine/Plugins/Runtime/InstancedActors/Source/InstancedActors/Public/InstancedActorsVisualizationTrait.h
+++ b/Engine/Plugins/Runtime/InstancedActors/Source/InstancedActors/Public/InstancedActorsVisualizationTrait.h
@@ -11,10 +11,12 @@ class UInstancedActorsData;
 
 /** 
  * Subclass of UMassStationaryVisualizationTrait which forces required settings for instanced actor entities and overrides
- * FMassRepresentationFragment.StaticMeshDescHandle to use a custom registered Visulization which reuses InstanceData's 
+ * FMassRepresentationFragment.StaticMeshDescHandle to use a custom registered Visualization which reuses InstanceData's 
  * ISMComponents via UMassRepresentationSubsystem::AddVisualDescWithISMComponent.
+ * Note that the trait is marked to not show up in class selection drop-downs. The reason is that this trait is supposed
+ * to be used internally by InstancedActors and is never expected to be a part of a user-authored entity config.
  */
-UCLASS(MinimalAPI)
+UCLASS(MinimalAPI, HideDropdown)
 class UInstancedActorsVisualizationTrait : public UMassStationaryDistanceVisualizationTrait
 {
 	GENERATED_BODY()
