Duplicate IAM ISMC MIDs for POI HLODs to prevent HLODs from referencing MIDs in non-HLOD ISMCs.

#rnx

diff --git a/Engine/Plugins/Runtime/InstancedActors/Source/InstancedActors/Private/InstancedActorsData.cpp b/Engine/Plugins/Runtime/InstancedActors/Source/InstancedActors/Private/InstancedActorsData.cpp
--- a/Engine/Plugins/Runtime/InstancedActors/Source/InstancedActors/Private/InstancedActorsData.cpp
+++ b/Engine/Plugins/Runtime/InstancedActors/Source/InstancedActors/Private/InstancedActorsData.cpp
@@ -634,6 +634,22 @@ bool UInstancedActorsData::SetInstanceTransform(const FInstancedActorsInstanceHa
 
 	return true;
 }
+
+void UInstancedActorsData::ForEachEditorPreviewISMC(TFunctionRef<bool(UInstancedStaticMeshComponent& /*ISMComponent*/)> InFunction) const
+{
+	// Iterate ISMComponents
+	for (UInstancedStaticMeshComponent* EditorPreviewISMComponent : EditorPreviewISMComponents)
+	{
+		if (IsValid(EditorPreviewISMComponent))
+		{
+			const bool bContinue = InFunction(*EditorPreviewISMComponent);
+			if (!bContinue)
+			{
+				break;
+			}
+		}
+	}
+}
 #endif // WITH_EDITOR
 
 FInstancedActorsInstanceIndex UInstancedActorsData::GetInstanceIndexForEntity(const FMassEntityHandle EntityHandle) const
diff --git a/Engine/Plugins/Runtime/InstancedActors/Source/InstancedActors/Public/InstancedActorsData.h b/Engine/Plugins/Runtime/InstancedActors/Source/InstancedActors/Public/InstancedActorsData.h
--- a/Engine/Plugins/Runtime/InstancedActors/Source/InstancedActors/Public/InstancedActorsData.h
+++ b/Engine/Plugins/Runtime/InstancedActors/Source/InstancedActors/Public/InstancedActorsData.h
@@ -177,6 +177,9 @@ public:
 	// @see IA.CompactInstances console command
 	bool RemoveInstance(const FInstancedActorsInstanceHandle& InstanceToRemove);
 
+	// Iterates all ISMCs created in GetOrCreateActorInstanceData.
+	// @param InFunction						The function to call for each ISMC.
+	void ForEachEditorPreviewISMC(TFunctionRef<bool(UInstancedStaticMeshComponent& /*ISMComponent*/)> InFunction) const;
 #endif
 
 	// Removes RuntimeRemoveInstances as if they were never present i.e: these removals are not persisted as
