[InstancedActors] Modified how InstancedActorsEditor module creates Actor menu extensions so that the behavior can be overridden.
Also, made a change that results in IAMs destruction when IA-to-Actor conversion is performed.

diff --git a/Engine/Plugins/Runtime/InstancedActors/Source/InstancedActorsEditor/Private/InstancedActorsEditorModule.cpp b/Engine/Plugins/Runtime/InstancedActors/Source/InstancedActorsEditor/Private/InstancedActorsEditorModule.cpp
--- a/Engine/Plugins/Runtime/InstancedActors/Source/InstancedActorsEditor/Private/InstancedActorsEditorModule.cpp
+++ b/Engine/Plugins/Runtime/InstancedActors/Source/InstancedActorsEditor/Private/InstancedActorsEditorModule.cpp
@@ -20,6 +20,7 @@ IMPLEMENT_MODULE(FInstancedActorsEditorModule, InstancedActorsEditor)
 
 void FInstancedActorsEditorModule::StartupModule()
 {
+	ResetConversionDelegates();
 	AddLevelViewportMenuExtender();
 }
 
@@ -29,6 +30,15 @@ void FInstancedActorsEditorModule::ShutdownModule()
 	RemoveLevelViewportMenuExtender();
 }
 
+void FInstancedActorsEditorModule::ResetConversionDelegates()
+{
+	ActorToIADelegate.BindRaw(this, &FInstancedActorsEditorModule::ConvertActorsToIAsUIAction);
+	IAToActorDelegate.BindRaw(this, &FInstancedActorsEditorModule::ConvertIAsToActorsUIAction);
+
+	ActorToIAFormatLabel = LOCTEXT("ConvertSelectedActorsToIAsText", "Convert {0} to Instanced Actors");
+	IAToActorFormatLabel = LOCTEXT("ConvertSelectedIAsToActorsText", "Convert {0}'s instances back to Actors");
+}
+
 void FInstancedActorsEditorModule::AddLevelViewportMenuExtender()
 {
 	if (!IsRunningGame())
@@ -62,7 +72,7 @@ TSharedRef<FExtender> FInstancedActorsEditorModule::CreateLevelViewportContextMe
 	TSharedRef<FExtender> Extender = MakeShareable(new FExtender);
 	if (InActors.Num() > 0)
 	{
-		FText ActorName = InActors.Num() == 1 ? FText::Format(LOCTEXT("ActorNameSingular", "\"{0}\""), FText::FromString(InActors[0]->GetActorLabel())) : LOCTEXT("ActorNamePlural", "Actors");
+		const FText ActorName = InActors.Num() == 1 ? FText::Format(LOCTEXT("ActorNameSingular", "\"{0}\""), FText::FromString(InActors[0]->GetActorLabel())) : LOCTEXT("ActorNamePlural", "Actors");
 
 		FLevelEditorModule& LevelEditor = FModuleManager::GetModuleChecked<FLevelEditorModule>(TEXT("LevelEditor"));
 		TSharedRef<FUICommandList> LevelEditorCommandBindings = LevelEditor.GetGlobalLevelEditorActions();
@@ -71,35 +81,41 @@ TSharedRef<FExtender> FInstancedActorsEditorModule::CreateLevelViewportContextMe
 		Extender->AddMenuExtension("ActorConvert", EExtensionHook::After, LevelEditorCommandBindings, FMenuExtensionDelegate::CreateLambda(
 			[this, ActorName, InActors](FMenuBuilder& MenuBuilder)
 			{
-				bool bCanExecuteActorToIAM = false;
-				bool bCanExecuteIAMtoActor = false;
+				bool bCanExecuteActorToIA = false;
+				bool bCanExecuteIAtoActor = false;
 				
 				// we can stop checking as soon as we know we have both cases in Selected Actors (the InActors array).
-				for (int32 Index = 0; Index < InActors.Num() && !(bCanExecuteActorToIAM && bCanExecuteIAMtoActor); ++Index)
+				for (int32 Index = 0; Index < InActors.Num() && !(bCanExecuteActorToIA && bCanExecuteIAtoActor); ++Index)
 				{
-					const bool bIsIAM = InActors[Index]->GetClass()->IsChildOf(AInstancedActorsManager::StaticClass());
-					// We can only convert and Actor to an IAM if it's not an IAM instance
-					bCanExecuteActorToIAM = bCanExecuteActorToIAM || !bIsIAM;
-					// We can only convert instances to and Actors only if it _is_ an IAM
-					bCanExecuteIAMtoActor = bCanExecuteIAMtoActor || bIsIAM;
+					const bool bIsIA = InActors[Index]->GetClass()->IsChildOf(AInstancedActorsManager::StaticClass());
+					// We can only convert and Actor to an IA if it's not an IA instance
+					bCanExecuteActorToIA = bCanExecuteActorToIA || !bIsIA;
+					// We can only convert instances to and Actors only if it _is_ an IA
+					bCanExecuteIAtoActor = bCanExecuteIAtoActor || bIsIA;
 				}
 
 				MenuBuilder.AddMenuEntry(
-					FText::Format(LOCTEXT("ConvertSelectedActorsToIAMsText", "Convert {0} to Instanced Actors"), ActorName),
-					LOCTEXT("ConvertSelectedActorsToIAMsTooltip", "Convert the selected actors to Instanced Actors instances."),
+					FText::Format(ActorToIAFormatLabel, ActorName),
+					LOCTEXT("ConvertSelectedActorsToIAsTooltip", "Convert the selected actors to Instanced Actors instances."),
 					FSlateIcon(FAppStyle::GetAppStyleSetName(), "Icons.Convert"),
 					FUIAction(
-						FExecuteAction::CreateRaw(this, &FInstancedActorsEditorModule::ConvertActorsToIAMsUIAction, InActors),
-						FCanExecuteAction::CreateLambda([bCanExecuteActorToIAM]() { return bCanExecuteActorToIAM; }))
+						FExecuteAction::CreateLambda([this, InActors]()
+							{
+								ActorToIADelegate.ExecuteIfBound(InActors);
+							}),
+						FCanExecuteAction::CreateLambda([bCanExecuteActorToIA]() { return bCanExecuteActorToIA; }))
 				);
 
 				MenuBuilder.AddMenuEntry(
-					FText::Format(LOCTEXT("ConvertSelectedIAMsToActorsText", "Convert {0}'s instances back to Actors"), ActorName),
-					LOCTEXT("ConvertSelectedIAMsToActorsToolTip", "Convert all the Instanced Actors instances back to Actors."),
+					FText::Format(IAToActorFormatLabel, ActorName),
+					LOCTEXT("ConvertSelectedIAsToActorsToolTip", "Convert all the Instanced Actors instances back to Actors."),
 					FSlateIcon(FAppStyle::GetAppStyleSetName(), "Icons.Convert"),
 					FUIAction(
-						FExecuteAction::CreateRaw(this, &FInstancedActorsEditorModule::ConvertIAMsToActorsUIAction, InActors),
-						FCanExecuteAction::CreateLambda([bCanExecuteIAMtoActor]() { return bCanExecuteIAMtoActor; }))
+						FExecuteAction::CreateLambda([this, InActors]()
+							{
+								IAToActorDelegate.ExecuteIfBound(InActors);
+							}),
+						FCanExecuteAction::CreateLambda([bCanExecuteIAtoActor]() { return bCanExecuteIAtoActor; }))
 				);
 			})
 		);
@@ -107,7 +123,13 @@ TSharedRef<FExtender> FInstancedActorsEditorModule::CreateLevelViewportContextMe
 	return Extender;
 }
 
-void FInstancedActorsEditorModule::ConvertActorsToIAMsUIAction(const TArray<AActor*> InActors) const
+void FInstancedActorsEditorModule::ConvertActorsToIAsUIAction(TConstArrayView<AActor*> InActors) const
+{
+	CustomizedConvertActorsToIAsUIAction(InActors, UInstancedActorsSubsystem::StaticClass());
+}
+
+void FInstancedActorsEditorModule::CustomizedConvertActorsToIAsUIAction(TConstArrayView<AActor*> InActors
+	, TSubclassOf<UInstancedActorsSubsystem> IASubsystemClass) const
 {
 	check(GEditor);
 
@@ -118,23 +140,23 @@ void FInstancedActorsEditorModule::ConvertActorsToIAMsUIAction(const TArray<AAct
 		return;
 	}
 
-	UInstancedActorsSubsystem* IAMSubsystem = World->GetSubsystem<UInstancedActorsSubsystem>();
-	check(IAMSubsystem);
+	UInstancedActorsSubsystem* IASubsystem = Cast<UInstancedActorsSubsystem>(World->GetSubsystemBase(IASubsystemClass));
+	check(IASubsystem);
 
-	FScopedTransaction Transaction(LOCTEXT("ConvertToIAM_Transaction", "Convert Actors to IAMs"));
+	FScopedTransaction Transaction(LOCTEXT("ConvertToIA_Transaction", "Convert Actors to IAs"));
 	GEditor->SelectNone(/*bNoteSelectionChange=*/true, /*bDeselectBSPSurfs=*/true, /*WarnAboutManyActors=*/false);
 
 	for (AActor* Actor : InActors)
 	{
-		// note that we skip all the IAMs here, we don't support converting IAM to instances.
-		// We can end up here if there are multiple different AActors selected and some of them are IAMs - the option
+		// note that we skip all the IAs here, we don't support converting IA to instances.
+		// We can end up here if there are multiple different AActors selected and some of them are IAs - the option
 		// to convert will still appear in the "Actor" menu.
-		if (Actor == nullptr || InActors[0]->GetClass()->IsChildOf(AInstancedActorsManager::StaticClass()))
+		if (Actor == nullptr || Actor->GetClass()->IsChildOf(AInstancedActorsManager::StaticClass()))
 		{
 			continue;
 		}
 
-		FInstancedActorsInstanceHandle InstanceHandle = IAMSubsystem->InstanceActor(Actor->GetClass(), Actor->GetActorTransform(), World->GetCurrentLevel());
+		FInstancedActorsInstanceHandle InstanceHandle = IASubsystem->InstanceActor(Actor->GetClass(), Actor->GetActorTransform(), World->GetCurrentLevel());
 		if (InstanceHandle.IsValid())
 		{
 			Actor->Destroy(); // This will call modify too.
@@ -144,7 +166,7 @@ void FInstancedActorsEditorModule::ConvertActorsToIAMsUIAction(const TArray<AAct
 	}
 }
 
-void FInstancedActorsEditorModule::ConvertIAMsToActorsUIAction(const TArray<AActor*> InActors) const
+void FInstancedActorsEditorModule::ConvertIAsToActorsUIAction(TConstArrayView<AActor*> InActors) const
 {
 	check(GEditor);
 
@@ -155,10 +177,10 @@ void FInstancedActorsEditorModule::ConvertIAMsToActorsUIAction(const TArray<AAct
 		return;
 	}
 
-	UInstancedActorsSubsystem* IAMSubsystem = World->GetSubsystem<UInstancedActorsSubsystem>();
-	check(IAMSubsystem);
+	UInstancedActorsSubsystem* IASubsystem = World->GetSubsystem<UInstancedActorsSubsystem>();
+	check(IASubsystem);
 
-	FScopedTransaction Transaction(LOCTEXT("ConvertToActorsFromIAM_Transaction", "Convert IAMs to Actors"));
+	FScopedTransaction Transaction(LOCTEXT("ConvertToActorsFromIA_Transaction", "Convert IAs to Actors"));
 	GEditor->SelectNone(/*bNoteSelectionChange=*/true, /*bDeselectBSPSurfs=*/true, /*WarnAboutManyActors=*/false);
 
 	FActorSpawnParameters SpawnParams;
@@ -203,7 +225,23 @@ void FInstancedActorsEditorModule::ConvertIAMsToActorsUIAction(const TArray<AAct
 					return true;
 				}
 				, IterationContext);
+
+			// since we removed all content from the instance we can just as well destroy it. 
+			Manager->Destroy();
 		}
 	}
 }
+
+void FInstancedActorsEditorModule::SetIAToActorDelegate(const FOnConvert& InDelegate, const FTextFormat& ActionFormatLabelOverride)
+{
+	IAToActorDelegate = InDelegate;
+	IAToActorFormatLabel = ActionFormatLabelOverride;
+}
+
+void FInstancedActorsEditorModule::SetActorToIADelegate(const FOnConvert& InDelegate, const FTextFormat& ActionFormatLabelOverride)
+{
+	ActorToIADelegate = InDelegate;
+	ActorToIAFormatLabel = ActionFormatLabelOverride;
+}
+
 #undef LOCTEXT_NAMESPACE
diff --git a/Engine/Plugins/Runtime/InstancedActors/Source/InstancedActorsEditor/Public/InstancedActorsEditorModule.h b/Engine/Plugins/Runtime/InstancedActors/Source/InstancedActorsEditor/Public/InstancedActorsEditorModule.h
--- a/Engine/Plugins/Runtime/InstancedActors/Source/InstancedActorsEditor/Public/InstancedActorsEditorModule.h
+++ b/Engine/Plugins/Runtime/InstancedActors/Source/InstancedActorsEditor/Public/InstancedActorsEditorModule.h
@@ -3,34 +3,67 @@
 #pragma  once
 
 #include "Containers/Array.h"
+#include "Containers/ArrayView.h"
 #include "Templates/SharedPointer.h"
 #include "Delegates/IDelegateInstance.h"
 #include "Modules/ModuleInterface.h"
+#include "Templates/SubclassOf.h"
 
 class AActor;
 class FExtender;
 class FUICommandList;
+class AInstancedActorsManager;
+class UInstancedActorsSubsystem;
 
 /**
 * The public interface to this module
+* This module contains logic for converting between AActors and Instanced Actors. 
+* @todo move to a dedicated class
 */
 class INSTANCEDACTORSEDITOR_API FInstancedActorsEditorModule : public IModuleInterface
 {
 public:
+	DECLARE_DELEGATE_OneParam(FOnConvert, TConstArrayView<AActor*>);
+
+	/** Convert selected actors to Instance Actors (IAs) while using the IASubsystemClass */
+	void CustomizedConvertActorsToIAsUIAction(TConstArrayView<AActor*> InActors, TSubclassOf<UInstancedActorsSubsystem> IASubsystemClass) const;
+	
+	/** 
+	 * Sets InDelegate as the delegate that will be executed as "Convert Instanced Actors to regular Actors" action.
+	 * @param ActionFormatLabelOverride if supplied will be used to override the default action label. Note that 
+	 *	this is a format string and needs to require one parameter, the Actor's name
+	 */
+	void SetIAToActorDelegate(const FOnConvert& InDelegate, const FTextFormat& ActionFormatLabelOverride = FTextFormat());
+
+	/**
+	 * Sets InDelegate as the delegate that will be executed as "Convert Actors to Instanced Actors" action.
+	 * @param ActionFormatLabelOverride if supplied will be used to override the default action label. Note that
+	 *	this is a format string and needs to require one parameter, the Actor's name
+	 */
+	void SetActorToIADelegate(const FOnConvert& InDelegate, const FTextFormat& ActionFormatLabelOverride = FTextFormat());
+
+	/** Resets FOnConvert delegates and action labels to their defaults */
+	void ResetConversionDelegates();
+
+protected:
 	// Begin IModuleInterface
 	virtual void StartupModule() override;
 	virtual void ShutdownModule() override;
 
-	/** Convert selected actors to Mass Instance Actors (IAMs). */
-	void ConvertActorsToIAMsUIAction(const TArray<AActor*> InActors) const;
+	/** Convert selected actors to Instance Actors (IAs). */
+	void ConvertActorsToIAsUIAction(TConstArrayView<AActor*> InActors) const;
 
-	/** Convert selected Mass Instance Actors to Actors. */
-	void ConvertIAMsToActorsUIAction(const TArray<AActor*> InActors) const;
+	/** Convert selected Instance Actors to Actors. */
+	void ConvertIAsToActorsUIAction(TConstArrayView<AActor*> InActors) const;
 
-protected:
 	TSharedRef<FExtender> CreateLevelViewportContextMenuExtender(const TSharedRef<FUICommandList> CommandList, const TArray<AActor*> InActors);
 	void AddLevelViewportMenuExtender();
 	void RemoveLevelViewportMenuExtender();
 
 	FDelegateHandle LevelViewportExtenderHandle;
+	FOnConvert ActorToIADelegate;
+	FOnConvert IAToActorDelegate;
+	FText CustomizedLabelPrefix;
+	FTextFormat ActorToIAFormatLabel;
+	FTextFormat IAToActorFormatLabel;
 };
