[InstancedActors] Added a short-circuit test to UInstancedActorsSubsystem::InstanceActor to bail out when input params are wrong

diff --git a/Engine/Plugins/Runtime/InstancedActors/Source/InstancedActors/Private/InstancedActorsSubsystem.cpp b/Engine/Plugins/Runtime/InstancedActors/Source/InstancedActors/Private/InstancedActorsSubsystem.cpp
--- a/Engine/Plugins/Runtime/InstancedActors/Source/InstancedActors/Private/InstancedActorsSubsystem.cpp
+++ b/Engine/Plugins/Runtime/InstancedActors/Source/InstancedActors/Private/InstancedActorsSubsystem.cpp
@@ -395,11 +395,20 @@ FInstancedActorsInstanceHandle UInstancedActorsSubsystem::InstanceActor(TSubclas
 
 FInstancedActorsInstanceHandle UInstancedActorsSubsystem::InstanceActor(TSubclassOf<AActor> ActorClass, FTransform InstanceTransform, ULevel* Level, const FGameplayTagContainer& InstanceTags, TSubclassOf<AInstancedActorsManager> ManagerClass)
 {
+	if (!ensureMsgf(Level, TEXT("Expecting a valid Level. Received nullptr.")))
+	{
+		return FInstancedActorsInstanceHandle();
+	}
+
 	UWorld* World = Level->GetWorld();
 	if (!ensureMsgf(!World->IsGameWorld(), TEXT("Instanced Actors doesn't yet support runtime addition of instances. Skipping instance creation")))
 	{
 		return FInstancedActorsInstanceHandle();
 	}
+	else if (!ensureMsgf(ActorClass, TEXT("Expecting a valid ActorClass. Received None.")))
+	{
+		return FInstancedActorsInstanceHandle();
+	}
 
 	if (!ManagerClass)
 	{
