[IAM] Clear template and exemplar actor data when IADs no longer need them.
* Add reference counting to template actors
* Cache TSharedPtr/TWeakPtr to structs that hold references to exemplar actors
* IADs release template actors upon deinit
* New FExemplarActorData destroys exemplar actor and request removal from the UInstancedActorsSubsystem map

diff --git a/Engine/Plugins/Runtime/InstancedActors/Source/InstancedActors/Private/InstancedActorsData.cpp b/Engine/Plugins/Runtime/InstancedActors/Source/InstancedActors/Private/InstancedActorsData.cpp
--- a/Engine/Plugins/Runtime/InstancedActors/Source/InstancedActors/Private/InstancedActorsData.cpp
+++ b/Engine/Plugins/Runtime/InstancedActors/Source/InstancedActors/Private/InstancedActorsData.cpp
@@ -150,6 +150,13 @@ namespace UE::InstancedActors
 			TEXT("IA.UpdateNextTickTimeFragments"),
 			bUpdateNextTickTimeFragments,
 			TEXT("If enabled (default) forces LOD updated for newly re-spawned entities (i.e. spawned by IAD that has been reloaded)"),
+			ECVF_Default);		
+
+		bool bEnableReleasingEntityTemplatesAndExemplarActors = true;
+		FAutoConsoleVariableRef CVarEnableReleasingEntityTemplatesAndExemplarActors(
+			TEXT("IA.EnableReleasingEntityTemplatesAndExemplarActors"),
+			bEnableReleasingEntityTemplatesAndExemplarActors,
+			TEXT("If enabled (default) instanced actors will release created entity templates and exemplar actors from central systems. Allows resources that are no longer in use to be garbage collected."),
 			ECVF_Default);
 	} // CVars
 
@@ -164,6 +171,22 @@ namespace UE::InstancedActors
 			InstanceTransform.SetIdentityZeroScale();
 		}
 	} // Helpers
+
+	//-----------------------------------------------------------------------------
+	// FExemplarActorData
+	//-----------------------------------------------------------------------------
+	FExemplarActorData::~FExemplarActorData()
+	{
+		if (CVars::bEnableReleasingEntityTemplatesAndExemplarActors)
+		{
+			AActor* ExemplarActor = Actor.Get();
+			check(ExemplarActor);
+	
+			InstancedActorsSubsystem.UnregisterExemplarActorClass(ExemplarActor->GetClass());
+
+			ExemplarActor->Destroy();
+		}
+	}
 } // namespace InstancedActors
 
 //-----------------------------------------------------------------------------
@@ -217,20 +240,33 @@ void UInstancedActorsData::Initialize()
 	}
 
 	// Get or create exemplar actor to derive entities from
-	const AActor& ExemplarActor = InstancedActorSubsystem.GetOrCreateExemplarActor(ActorClass);
+	check(!ExemplarActorData.IsValid());
+	ExemplarActorData = InstancedActorSubsystem.GetOrCreateExemplarActor(ActorClass);
+	const AActor* const ExemplarActor = ExemplarActorData->Actor.Get();
+	check(ExemplarActor);
 
 	// Add default visualization at index 0
 	//FInstancedActorsVisualizationDesc DefaultVisualiation = FInstancedActorsVisualizationDesc::FromActor(ExemplarActor, &UE::InstancedActors::VisualizationDescrFromActorAdditionalSteps);
-	FInstancedActorsVisualizationDesc DefaultVisualiation = InstancedActorSubsystem.CreateVisualDescriptionFromActor(ExemplarActor);
+	FInstancedActorsVisualizationDesc DefaultVisualiation = InstancedActorSubsystem.CreateVisualDescriptionFromActor(*ExemplarActor);
 	const uint8 VisualizationIndex = AddVisualization(DefaultVisualiation);
 	check(VisualizationIndex == 0);
 
 	// Create entity template
-	CreateEntityTemplate(ExemplarActor);
+	CreateEntityTemplate(*ExemplarActor);
 
 	bHasEverInitialized = true;
 }
 
+void UInstancedActorsData::Deinitialize()
+{
+	if (bHasEverInitialized && UE::InstancedActors::CVars::bEnableReleasingEntityTemplatesAndExemplarActors)
+	{
+		ReleaseEntityTemplate();
+
+		ExemplarActorData.Reset();
+	}
+}
+
 void UInstancedActorsData::CreateEntityTemplate(const AActor& ExemplarActor)
 {
 	UWorld* World = GetWorld();
@@ -238,7 +274,7 @@ void UInstancedActorsData::CreateEntityTemplate(const AActor& ExemplarActor)
 
 	FMassEntityManager& MassEntityManager = GetMassEntityManagerChecked();
 
-	FMassEntityConfig EntityConfig(*this);
+	EntityConfig = FMassEntityConfig(*this);
 
 	UMassDistanceLODCollectorTrait* LODCollectorTrait = NewObject<UMassDistanceLODCollectorTrait>(this);
 	EntityConfig.AddTrait(*LODCollectorTrait);
@@ -251,7 +287,7 @@ void UInstancedActorsData::CreateEntityTemplate(const AActor& ExemplarActor)
 	EntityConfig.AddTrait(*VisTrait);
 
 	// Allow UInstancedActorsComponent's to extend entity config
-	ExemplarActor.ForEachComponent<UInstancedActorsComponent>(/*bIncludeFromChildActors*/ false, [this, &MassEntityManager, &EntityConfig](const UInstancedActorsComponent* InstancedActorComponent)
+	ExemplarActor.ForEachComponent<UInstancedActorsComponent>(/*bIncludeFromChildActors*/ false, [this, &MassEntityManager](const UInstancedActorsComponent* InstancedActorComponent)
 		{ InstancedActorComponent->ModifyMassEntityConfig(MassEntityManager, this, EntityConfig); });
 
 	const FMassEntityTemplate& BaseEntityTemplate = EntityConfig.GetOrCreateEntityTemplate(*World);
@@ -282,6 +318,14 @@ void UInstancedActorsData::ModifyEntityTemplate(FMassEntityTemplateData& Modifie
 	// ModifiedTemplate.RemoveTag<FMassStationaryISMSwitcherProcessorTag>();
 }
 
+void UInstancedActorsData::ReleaseEntityTemplate()
+{
+	UWorld* World = GetWorld();
+	check(World);
+
+	EntityConfig.DestroyEntityTemplate(*World);
+}
+
 void UInstancedActorsData::SpawnEntities()
 {
 	if (NumValidInstances <= 0)
diff --git a/Engine/Plugins/Runtime/InstancedActors/Source/InstancedActors/Private/InstancedActorsManager.cpp b/Engine/Plugins/Runtime/InstancedActors/Source/InstancedActors/Private/InstancedActorsManager.cpp
--- a/Engine/Plugins/Runtime/InstancedActors/Source/InstancedActors/Private/InstancedActorsManager.cpp
+++ b/Engine/Plugins/Runtime/InstancedActors/Source/InstancedActors/Private/InstancedActorsManager.cpp
@@ -304,6 +304,12 @@ void AInstancedActorsManager::EndPlay(const EEndPlayReason::Type EndPlayReason)
 {
 	DespawnAllEntities();
 
+	for (TObjectPtr<UInstancedActorsData>& InstanceData : PerActorClassInstanceData)
+	{
+		check(InstanceData);
+		InstanceData->Deinitialize();
+	}
+
 	// Deregister with UInstancedActorsSubsystem
 	if (InstancedActorSubsystem)
 	{
@@ -518,8 +524,10 @@ void AInstancedActorsManager::SerializeInstancePersistenceData(FStructuredArchiv
 	int32 NumPersistedIACs = 0;
 	if (InstanceData)
 	{
-		const AActor& ExemplarActor = GetInstancedActorSubsystemChecked().GetOrCreateExemplarActor(InstanceData->ActorClass);
-		ExemplarActor.ForEachComponent<UInstancedActorsComponent>(/*bIncludeFromChildActors*/ false, [&](UInstancedActorsComponent* InstancedActorComponent)
+		TSharedPtr<UE::InstancedActors::FExemplarActorData> ExemplarActorData = GetInstancedActorSubsystemChecked().GetOrCreateExemplarActor(InstanceData->ActorClass);
+		const AActor* const ExemplarActor = ExemplarActorData->Actor.Get();
+		check(ExemplarActor);
+		ExemplarActor->ForEachComponent<UInstancedActorsComponent>(/*bIncludeFromChildActors*/ false, [&](UInstancedActorsComponent* InstancedActorComponent)
 			{
 				// Skip IAC's that don't want / need serialization. When writing this ensures we don't waste space by writing empty entries.
 				if (InstancedActorComponent->ShouldSerializeInstancePersistenceData(UnderlyingArchive, InstanceData, TimeDelta))
@@ -694,11 +702,17 @@ UInstancedActorsData& AInstancedActorsManager::GetOrCreateActorInstanceData(TSub
 	check(World);
 	UInstancedActorsSubsystem* EditorInstancedActorSubsystem = UE::InstancedActors::Utils::GetInstancedActorsSubsystem(*World);
 	check(EditorInstancedActorSubsystem);
-	const AActor& ExemplarActor = EditorInstancedActorSubsystem->GetOrCreateExemplarActor(ActorClass);
+	// @todo For performance reasons it would be helpful to be able to cache this into UInstancedActorsData::ExemplarActorData
+	// This would mean that undo-ing a conversion of actors to IAM in Editor will not clear or destroy the exemplar actor/map entry,
+	// which is the reason for the sanity-check of the actor ptr in UInstancedActorsSubsystem::GetOrCreateExemplarActor().
+	// We need to also investigate why converting the IAM to actors (destroying the IAM) does not seem to trigger GC on the UInstancedActorsData
+	TSharedPtr<UE::InstancedActors::FExemplarActorData> ExemplarActorData = EditorInstancedActorSubsystem->GetOrCreateExemplarActor(ActorClass);
+	const AActor* const ExemplarActor = ExemplarActorData->Actor.Get();
+	check(ExemplarActor);
 
 	// Create stand-in 'editor only' ISMC's to preview instances in the level editor. These will be stripped
 	// out at cook time and in PostLoad for game worlds (PIE)
-	FInstancedActorsVisualizationDesc DefaultVisualiation = EditorInstancedActorSubsystem->CreateVisualDescriptionFromActor(ExemplarActor);
+	FInstancedActorsVisualizationDesc DefaultVisualiation = EditorInstancedActorSubsystem->CreateVisualDescriptionFromActor(*ExemplarActor);
 	CreateISMComponents(DefaultVisualiation, NewInstanceData->SharedSettings, NewInstanceData->EditorPreviewISMComponents, /*bEditorPreviewISMCs*/ true);
 
 	// Add EditorPreviewISMComponents to the known ISMComponent-to-InstanceData map (ISMComponentToInstanceDataMap)
diff --git a/Engine/Plugins/Runtime/InstancedActors/Source/InstancedActors/Private/InstancedActorsSubsystem.cpp b/Engine/Plugins/Runtime/InstancedActors/Source/InstancedActors/Private/InstancedActorsSubsystem.cpp
--- a/Engine/Plugins/Runtime/InstancedActors/Source/InstancedActors/Private/InstancedActorsSubsystem.cpp
+++ b/Engine/Plugins/Runtime/InstancedActors/Source/InstancedActors/Private/InstancedActorsSubsystem.cpp
@@ -619,25 +619,28 @@ bool UInstancedActorsSubsystem::HasInstancesOfClass(const FBox& QueryBounds, TSu
 	return bHasInstances;
 }
 
-AActor& UInstancedActorsSubsystem::GetOrCreateExemplarActor(TSubclassOf<AActor> ActorClass)
+TSharedRef<UE::InstancedActors::FExemplarActorData> UInstancedActorsSubsystem::GetOrCreateExemplarActor(TSubclassOf<AActor> ActorClass)
 {
 	TRACE_CPUPROFILER_EVENT_SCOPE(UInstancedActorsSubsystem GetOrCreateExemplarActor);
 
-	check(ActorClass.Get());
+	const UClass* const ActorClassPtr = ActorClass.Get();
+	check(ActorClassPtr);
 
 	// Return existing?
-	TObjectPtr<AActor>* CachedExemplarActor = ExemplarActors.Find(ActorClass.Get());
-	if (CachedExemplarActor)
+	const uint32 ActorClassHash = GetTypeHash(TObjectKey<const UClass>(ActorClassPtr));
+	
+	if (const TWeakPtr<UE::InstancedActors::FExemplarActorData>* CachedExemplarActorDataPtr = ExemplarActors.FindByHash(ActorClassHash, ActorClassPtr))
 	{
 		// This can fail in editor with undo/redo in the mix.
-		if (IsValid(*CachedExemplarActor))
+		TSharedPtr<UE::InstancedActors::FExemplarActorData> CachedExemplarActorData = CachedExemplarActorDataPtr->Pin();
+		if (CachedExemplarActorData.IsValid() && CachedExemplarActorData->Actor.Get() != nullptr)
 		{
-			return **CachedExemplarActor;
+			return CachedExemplarActorData.ToSharedRef();
 		}
 		else
 		{
 			// The examplar is not valid, we'll remove it and then re-create it below.
-			ExemplarActors.Remove(ActorClass.Get());
+			ExemplarActors.RemoveByHash(ActorClassHash, ActorClassPtr);
 		}
 	}
 
@@ -685,9 +688,22 @@ AActor& UInstancedActorsSubsystem::GetOrCreateExemplarActor(TSubclassOf<AActor>
 	check(NewExemplarActor);
 
 	// Cache for subsequent calls
-	ExemplarActors.Add(ActorClass.Get(), NewExemplarActor);
+	TSharedPtr<UE::InstancedActors::FExemplarActorData> NewExemplarActorDataPtr{new UE::InstancedActors::FExemplarActorData{*NewExemplarActor, *this}};
+	ExemplarActors.AddByHash(ActorClassHash, ActorClassPtr, NewExemplarActorDataPtr);
+
+	return NewExemplarActorDataPtr.ToSharedRef();
+}
 
-	return *NewExemplarActor;
+void UInstancedActorsSubsystem::UnregisterExemplarActorClass(TSubclassOf<AActor> ActorClass)
+{
+	const UClass* const ActorClassPtr = ActorClass.Get();
+	check(ActorClassPtr);
+	
+	const uint32 ActorClassHash = GetTypeHash(TObjectKey<const UClass>(ActorClassPtr));	
+	if (const TWeakPtr<UE::InstancedActors::FExemplarActorData>* CachedExemplarActorDataPtr = ExemplarActors.FindByHash(ActorClassHash, ActorClassPtr))
+	{
+		ExemplarActors.RemoveByHash(ActorClassHash, ActorClassPtr);
+	}
 }
 
 FSharedStruct UInstancedActorsSubsystem::GetOrCompileSettingsForActorClass(TSubclassOf<AActor> ActorClass)
@@ -926,4 +942,4 @@ bool UInstancedActorsSubsystem::RegisterNewSharedFragmentsInternal(TConstStructV
 	check(SortedSharedFragments.Num() == AllSharedFragmentsOfType.Num());
 
 	return bParamStructFound;
-}
+}
\ No newline at end of file
diff --git a/Engine/Plugins/Runtime/InstancedActors/Source/InstancedActors/Public/InstancedActorsData.h b/Engine/Plugins/Runtime/InstancedActors/Source/InstancedActors/Public/InstancedActorsData.h
--- a/Engine/Plugins/Runtime/InstancedActors/Source/InstancedActors/Public/InstancedActorsData.h
+++ b/Engine/Plugins/Runtime/InstancedActors/Source/InstancedActors/Public/InstancedActorsData.h
@@ -4,6 +4,7 @@
 
 #include "InstancedActorsTypes.h"
 #include "InstancedActorsReplication.h"
+#include "MassEntityConfigAsset.h"
 #include "MassEntityTemplate.h"
 
 #include "InstancedActorsData.generated.h"
@@ -11,6 +12,26 @@
 class AInstancedActorsManager;
 struct FInstancedActorsSettings;
 
+namespace UE::InstancedActors
+{
+struct FExemplarActorData
+{
+	~FExemplarActorData();
+
+	TObjectPtr<AActor> Actor;
+private:
+	FExemplarActorData(AActor& InActor, UInstancedActorsSubsystem& InInstancedActorSubsystem)
+		: Actor(InActor)
+		, InstancedActorsSubsystem(InInstancedActorSubsystem)
+	{
+	}
+
+	UInstancedActorsSubsystem& InstancedActorsSubsystem;
+
+	friend class ::UInstancedActorsSubsystem;
+};
+} // UE::InstancedActors
+
 // @todo there's a lot of public variables in this class, and properties are mixed with functions. A refactor is coming soon.
 
 /**
@@ -37,6 +58,9 @@ public:
 	// then despawn all Mass entities and reset any other runtime instance data
 	void DespawnEntities();
 
+	// Called from the parent's AInstancedActorsManager::EndPlay. Can release the entity template and exemplar Actor from memory 
+	void Deinitialize();
+
 	UFUNCTION(BlueprintPure, Category = InstancedActors)
 	AInstancedActorsManager* GetManager() const;
 
@@ -381,6 +405,9 @@ protected:
 	void CreateEntityTemplate(const AActor& ExemplarActor);
 	virtual void ModifyEntityTemplate(FMassEntityTemplateData& ModifiedTemplate, const AActor& ExemplarActor);
 
+	// Called from Deinitialize to destroy the owned FMassEntityTemplate  
+	void ReleaseEntityTemplate();
+
 	// Helper function used in ApplyInstanceDeltas to apply a single delta
 	// @see ApplyInstanceDeltas
 	virtual void ApplyInstanceDelta(FMassEntityManager& EntityManager, const FInstancedActorsDelta& InstanceDelta, TArray<FInstancedActorsInstanceIndex>& OutEntitiesToRemove);
@@ -449,4 +476,9 @@ private:
 	// @see Serialize
 	UPROPERTY(Replicated, SaveGame, Transient)
 	FInstancedActorsDeltaList InstanceDeltas;
+
+	UPROPERTY(Transient)
+	FMassEntityConfig EntityConfig;
+
+	TSharedPtr<UE::InstancedActors::FExemplarActorData> ExemplarActorData;
 };
diff --git a/Engine/Plugins/Runtime/InstancedActors/Source/InstancedActors/Public/InstancedActorsSubsystem.h b/Engine/Plugins/Runtime/InstancedActors/Source/InstancedActors/Public/InstancedActorsSubsystem.h
--- a/Engine/Plugins/Runtime/InstancedActors/Source/InstancedActors/Public/InstancedActorsSubsystem.h
+++ b/Engine/Plugins/Runtime/InstancedActors/Source/InstancedActors/Public/InstancedActorsSubsystem.h
@@ -22,6 +22,11 @@ struct FInstancedActorsManagerHandle;
 struct FInstancedActorsModifierVolumeHandle;
 class UInstancedActorsProjectSettings;
 
+namespace UE::InstancedActors
+{
+struct FExemplarActorData;
+} // UE::InstancedActors
+
 /**
  * Instanced Actor subsystem used to spawn AInstancedActorsManager's and populate their instance data.
  * It also keeps track of all InstancedActorDatas and can be queried for them
@@ -117,8 +122,13 @@ public:
 	 * the main game world.
 	 *
 	 * These 'exemplar' actors are fully constructed, including BP construction scripts up to (but not including) BeginPlay.
-	 */ 
-	INSTANCEDACTORS_API AActor& GetOrCreateExemplarActor(TSubclassOf<AActor> ActorClass);
+	 */
+	INSTANCEDACTORS_API TSharedRef<UE::InstancedActors::FExemplarActorData> GetOrCreateExemplarActor(TSubclassOf<AActor> ActorClass);
+	
+	/**
+	 * Removes exemplar actor class from the map
+	 */
+	INSTANCEDACTORS_API void UnregisterExemplarActorClass(TSubclassOf<AActor> ActorClass);
 
 	/** 
 	 * Compiles and caches finalized settings for ActorClass based off FInstancedActorsClassSettingsBase found in 
@@ -254,7 +264,7 @@ protected:
 
 	// Lazily created exemplar actors for instance actor classes
 	// @see GetOrCreateExemplarActor
-	TMap<TObjectKey<const UClass>, TObjectPtr<AActor>> ExemplarActors;
+	TMap<TObjectKey<const UClass>, TWeakPtr<UE::InstancedActors::FExemplarActorData>> ExemplarActors;
 
 	UPROPERTY(Transient)
 	TObjectPtr<const UScriptStruct> SettingsType;
diff --git a/Engine/Plugins/Runtime/MassGameplay/Source/MassRepresentation/Private/MassDistanceVisualizationTrait.cpp b/Engine/Plugins/Runtime/MassGameplay/Source/MassRepresentation/Private/MassDistanceVisualizationTrait.cpp
--- a/Engine/Plugins/Runtime/MassGameplay/Source/MassRepresentation/Private/MassDistanceVisualizationTrait.cpp
+++ b/Engine/Plugins/Runtime/MassGameplay/Source/MassRepresentation/Private/MassDistanceVisualizationTrait.cpp
@@ -97,4 +97,11 @@ void UMassDistanceVisualizationTrait::BuildTemplate(FMassEntityTemplateBuildCont
 	BuildContext.AddTag<FMassVisualizationProcessorTag>();
 }
 
-
+void UMassDistanceVisualizationTrait::DestroyTemplate(const UWorld& World) const
+{
+	if (UMassRepresentationSubsystem* RepresentationSubsystem = Cast<UMassRepresentationSubsystem>(World.GetSubsystemBase(RepresentationSubsystemClass)))
+	{
+		RepresentationSubsystem->ReleaseTemplate(HighResTemplateActor);
+		RepresentationSubsystem->ReleaseTemplate(LowResTemplateActor);
+	}
+}
diff --git a/Engine/Plugins/Runtime/MassGameplay/Source/MassRepresentation/Private/MassRepresentationSubsystem.cpp b/Engine/Plugins/Runtime/MassGameplay/Source/MassRepresentation/Private/MassRepresentationSubsystem.cpp
--- a/Engine/Plugins/Runtime/MassGameplay/Source/MassRepresentation/Private/MassRepresentationSubsystem.cpp
+++ b/Engine/Plugins/Runtime/MassGameplay/Source/MassRepresentation/Private/MassRepresentationSubsystem.cpp
@@ -74,11 +74,18 @@ void UMassRepresentationSubsystem::DirtyStaticMeshInstances()
 int16 UMassRepresentationSubsystem::FindOrAddTemplateActor(const TSubclassOf<AActor>& ActorClass)
 {
 	UE_MT_SCOPED_WRITE_ACCESS(TemplateActorsMTAccessDetector);
-	int32 VisualIndex = TemplateActors.Find(ActorClass);
+
+	int32 VisualIndex = TemplateActors.IndexOfByPredicate(FTemplateActorEqualsPredicate{ActorClass});
+
 	if (VisualIndex == INDEX_NONE)
 	{
-		VisualIndex = TemplateActors.Add(ActorClass);
+		VisualIndex = TemplateActors.Add({ActorClass, 1u});
+	}
+	else
+	{
+		++TemplateActors[VisualIndex].RefCount;
 	}
+
 	check(VisualIndex < INT16_MAX);
 	return (int16)VisualIndex;
 }
@@ -96,7 +103,7 @@ AActor* UMassRepresentationSubsystem::GetOrSpawnActorFromTemplate(const FMassEnt
 
 	//@todo: this would be a good place to do pooling of actors instead of spawning them every time
 	check(ActorSpawnerSubsystem);
-	const TSubclassOf<AActor> TemplateToSpawn = TemplateActors[TemplateActorIndex];
+	const TSubclassOf<AActor> TemplateToSpawn = TemplateActors[TemplateActorIndex].Actor;
 
 	if (InOutSpawnRequestHandle.IsValid())
 	{
@@ -174,7 +181,7 @@ TSubclassOf<AActor> UMassRepresentationSubsystem::GetTemplateActorClass(const in
 		return nullptr;
 	}
 
-	return TemplateActors[TemplateActorIndex];
+	return TemplateActors[TemplateActorIndex].Actor;
 }
 
 bool UMassRepresentationSubsystem::IsCollisionLoaded(const FName TargetGrid, const FTransform& Transform) const
@@ -204,6 +211,26 @@ bool UMassRepresentationSubsystem::IsCollisionLoaded(const FName TargetGrid, con
 	return WorldPartitionSubsystem->IsStreamingCompleted(EWorldPartitionRuntimeCellState::Activated, QuerySources, /*bExactState*/ false);
 }
 
+void UMassRepresentationSubsystem::ReleaseTemplate(const TSubclassOf<AActor>& ActorClass)
+{
+	if (ActorClass)
+	{
+		UE_MT_SCOPED_WRITE_ACCESS(TemplateActorsMTAccessDetector);
+		
+		const int32 TemplateActorIndex = TemplateActors.IndexOfByPredicate(FTemplateActorEqualsPredicate{ActorClass});
+		check(TemplateActors.IsValidIndex(TemplateActorIndex));
+
+		FTemplateActorData& TemplateActorData = TemplateActors[TemplateActorIndex];
+		check(TemplateActorData.RefCount > 0u);
+		--TemplateActorData.RefCount;
+
+		if (TemplateActorData.RefCount == 0u)
+		{
+			TemplateActors.RemoveAt(TemplateActorIndex);
+		}
+	}
+}
+
 void UMassRepresentationSubsystem::ReleaseAllResources()
 {
 	ensureMsgf(HandledMassAgents.Num() == 0, TEXT("All entities must be released before releasing resources"));
@@ -275,7 +302,7 @@ bool UMassRepresentationSubsystem::ReleaseTemplateActorInternal(const int16 Temp
 		UE_LOG(LogMassRepresentation, Error, TEXT("Template actor type %i is not referring a valid type"), TemplateActorIndex);
 		return false;
 	}
-	const TSubclassOf<AActor> TemplateToRelease = TemplateActors[TemplateActorIndex];
+	const TSubclassOf<AActor> TemplateToRelease = TemplateActors[TemplateActorIndex].Actor;
 
 	// We can only release existing and matching template actors
 	if (!ActorToRelease || ActorToRelease->GetClass() != TemplateToRelease)
@@ -294,7 +321,7 @@ bool UMassRepresentationSubsystem::CancelSpawningInternal(const int16 TemplateAc
 		UE_LOG(LogMassRepresentation, Error, TEXT("Template actor type %i is not referring a valid type"), TemplateActorIndex);
 		return false;
 	}
-	const TSubclassOf<AActor> TemplateToRelease = TemplateActors[TemplateActorIndex];
+	const TSubclassOf<AActor> TemplateToRelease = TemplateActors[TemplateActorIndex].Actor;
 
 	check(ActorSpawnerSubsystem);
 	// Check if there is something to cancel
@@ -331,7 +358,7 @@ bool UMassRepresentationSubsystem::DoesActorMatchTemplate(const AActor& Actor, c
 		return false;
 	}
 
-	const TSubclassOf<AActor> Template = TemplateActors[TemplateActorIndex];
+	const TSubclassOf<AActor> Template = TemplateActors[TemplateActorIndex].Actor;
 	return Actor.GetClass() == Template;
 }
 
@@ -452,3 +479,14 @@ void UMassRepresentationSubsystem::OnMassAgentComponentEntityDetaching(const UMa
 		}
 	}
 }
+
+void UMassRepresentationSubsystem::AddReferencedObjects(UObject* InThis, FReferenceCollector& Collector)
+{
+	Super::AddReferencedObjects(InThis, Collector);
+
+	UMassRepresentationSubsystem& TypedThis = *CastChecked<UMassRepresentationSubsystem>(InThis);
+	for (FTemplateActorData& TemplateActorData : TypedThis.TemplateActors)
+	{
+		Collector.AddStableReference(&TemplateActorData.Actor.GetGCPtr());
+	}
+}
diff --git a/Engine/Plugins/Runtime/MassGameplay/Source/MassRepresentation/Public/MassDistanceVisualizationTrait.h b/Engine/Plugins/Runtime/MassGameplay/Source/MassRepresentation/Public/MassDistanceVisualizationTrait.h
--- a/Engine/Plugins/Runtime/MassGameplay/Source/MassRepresentation/Public/MassDistanceVisualizationTrait.h
+++ b/Engine/Plugins/Runtime/MassGameplay/Source/MassRepresentation/Public/MassDistanceVisualizationTrait.h
@@ -22,6 +22,7 @@ public:
 	UMassDistanceVisualizationTrait();
 
 	virtual void BuildTemplate(FMassEntityTemplateBuildContext& BuildContext, const UWorld& World) const override;
+	virtual void DestroyTemplate(const UWorld& World) const override;
 
 	/** Instanced static mesh information for this agent */
 	UPROPERTY(EditAnywhere, Category = "Mass|Visual")
diff --git a/Engine/Plugins/Runtime/MassGameplay/Source/MassRepresentation/Public/MassRepresentationSubsystem.h b/Engine/Plugins/Runtime/MassGameplay/Source/MassRepresentation/Public/MassRepresentationSubsystem.h
--- a/Engine/Plugins/Runtime/MassGameplay/Source/MassRepresentation/Public/MassRepresentationSubsystem.h
+++ b/Engine/Plugins/Runtime/MassGameplay/Source/MassRepresentation/Public/MassRepresentationSubsystem.h
@@ -141,6 +141,11 @@ public:
 
 	bool IsCollisionLoaded(const FName TargetGrid, const FTransform& Transform) const;
 
+	/**
+	 * Responds to the FMassEntityTemplate getting destroyed, and releases reference to corresponding Actor in TemplateActors
+	 */
+	void ReleaseTemplate(const TSubclassOf<AActor>& ActorClass);
+
 	/**
 	 * Release all references to static meshes and template actors
 	 * Use with caution, all entities using this representation subsystem must be destroy otherwise they will point to invalid resources */
@@ -163,11 +168,30 @@ protected:
 	bool ReleaseTemplateActorInternal(const int16 TemplateActorIndex, AActor* ActorToRelease, bool bImmediate);
 	bool CancelSpawningInternal(const int16 TemplateActorIndex, FMassActorSpawnRequestHandle& SpawnRequestHandle);
 
+	static void AddReferencedObjects(UObject* InThis, FReferenceCollector& Collector);
+
 protected:
 
+	struct FTemplateActorData
+	{
+		TSubclassOf<AActor> Actor;
+		uint32 RefCount{0u};
+	};
+	
+	struct FTemplateActorEqualsPredicate
+	{
+		const TSubclassOf<AActor>& ActorClass;
+
+		FTemplateActorEqualsPredicate(const TSubclassOf<AActor>& ActorClass) : ActorClass(ActorClass) {}
+
+		bool operator()(const FTemplateActorData& ActorData) const
+		{
+			return ActorData.Actor == ActorClass;
+		}
+	};
+
 	/** The array of all the template actors */
-	UPROPERTY(Transient)
-	TArray<TSubclassOf<AActor>> TemplateActors;
+	TSparseArray<FTemplateActorData> TemplateActors;
 	UE_MT_DECLARE_RW_ACCESS_DETECTOR(TemplateActorsMTAccessDetector);
 
 	/** The component that handles all the static mesh instances */
diff --git a/Engine/Plugins/Runtime/MassGameplay/Source/MassSpawner/Private/MassEntityConfigAsset.cpp b/Engine/Plugins/Runtime/MassGameplay/Source/MassSpawner/Private/MassEntityConfigAsset.cpp
--- a/Engine/Plugins/Runtime/MassGameplay/Source/MassSpawner/Private/MassEntityConfigAsset.cpp
+++ b/Engine/Plugins/Runtime/MassGameplay/Source/MassSpawner/Private/MassEntityConfigAsset.cpp
@@ -93,7 +93,7 @@ void FMassEntityConfig::DestroyEntityTemplate(const UWorld& World) const
 	for (const UMassEntityTraitBase* Trait : CombinedTraits)
 	{
 		check(Trait);
-		Trait->DestroyTemplate();
+		Trait->DestroyTemplate(World);
 	}
 
 	// TODO - The templates are not being torn down completely, resulting in traits that leave data in various subsystems. (Representation system)
diff --git a/Engine/Plugins/Runtime/MassGameplay/Source/MassSpawner/Public/MassEntityTraitBase.h b/Engine/Plugins/Runtime/MassGameplay/Source/MassSpawner/Public/MassEntityTraitBase.h
--- a/Engine/Plugins/Runtime/MassGameplay/Source/MassSpawner/Public/MassEntityTraitBase.h
+++ b/Engine/Plugins/Runtime/MassGameplay/Source/MassSpawner/Public/MassEntityTraitBase.h
@@ -49,8 +49,11 @@ public:
 	/** Appends items into the entity template required for the trait. */
 	virtual void BuildTemplate(FMassEntityTemplateBuildContext& BuildContext, const UWorld& World) const PURE_VIRTUAL(UMassEntityTraitBase::BuildTemplate, return; );
 
+	UE_DEPRECATED(5.5, "This flavor of DestroyTemplate has been deprecated, use the version taking the World parameter")
 	virtual void DestroyTemplate() const;
 
+	virtual void DestroyTemplate(const UWorld& World) const {}
+
 	/**
 	 * Called once all traits have been processed and fragment requirements have been checked. Override this function
 	 * to perform additional Trait's configuration validation. Returning `false` will indicate that the trait instance
