[InstancedActors] Release UInstancedActorsData::ExemplarActorData reference when the parent AInstancedActorsManager gets reassigned to a new UInstancedActorsSubsystem. Also re-enable the check() in UInstancedActorsData::Initialize() that expects ExemplarActorData to be invalid (since this will happen after the UInstancedActorsSubsystem reassignment)

diff --git a/Engine/Plugins/Runtime/InstancedActors/Source/InstancedActors/Private/InstancedActorsData.cpp b/Engine/Plugins/Runtime/InstancedActors/Source/InstancedActors/Private/InstancedActorsData.cpp
--- a/Engine/Plugins/Runtime/InstancedActors/Source/InstancedActors/Private/InstancedActorsData.cpp
+++ b/Engine/Plugins/Runtime/InstancedActors/Source/InstancedActors/Private/InstancedActorsData.cpp
@@ -240,8 +240,7 @@ void UInstancedActorsData::Initialize()
 	}
 
 	// Get or create exemplar actor to derive entities from
-	// @todo this check() is triggering somehow in certain cases. Need to investigate the root cause, but disabling for now to unblock
-	//check(!ExemplarActorData.IsValid());
+	check(!ExemplarActorData.IsValid());
 	ExemplarActorData = InstancedActorSubsystem.GetOrCreateExemplarActor(ActorClass);
 	const AActor* const ExemplarActor = ExemplarActorData->Actor.Get();
 	check(ExemplarActor);
diff --git a/Engine/Plugins/Runtime/InstancedActors/Source/InstancedActors/Private/InstancedActorsManager.cpp b/Engine/Plugins/Runtime/InstancedActors/Source/InstancedActors/Private/InstancedActorsManager.cpp
--- a/Engine/Plugins/Runtime/InstancedActors/Source/InstancedActors/Private/InstancedActorsManager.cpp
+++ b/Engine/Plugins/Runtime/InstancedActors/Source/InstancedActors/Private/InstancedActorsManager.cpp
@@ -219,6 +219,14 @@ void AInstancedActorsManager::OnAddedToSubsystem(UInstancedActorsSubsystem& InIn
 		DespawnAllEntities();
 		// and let registered modifiers go
 		RemoveAllModifierVolumes();
+		
+		// Since the old IASubsystem is part of a different exemplar world, to be safe we should release the IADs' references to their
+		// exemplar actors, since they were spawned under that old world
+		for (TObjectPtr<UInstancedActorsData>& InstanceData : PerActorClassInstanceData)
+		{
+			check(InstanceData);
+			InstanceData->ExemplarActorData.Reset();
+		}
 
 		InstancedActorSubsystem->RemoveManager(ManagerHandle);
 		InstancedActorSubsystem = nullptr;
