[InstancedActors] Removed indications of project-specific jira tickets from a source code comment

diff --git a/Engine/Plugins/Runtime/InstancedActors/Source/InstancedActors/Public/InstancedActorsSettingsTypes.h b/Engine/Plugins/Runtime/InstancedActors/Source/InstancedActors/Public/InstancedActorsSettingsTypes.h
--- a/Engine/Plugins/Runtime/InstancedActors/Source/InstancedActors/Public/InstancedActorsSettingsTypes.h
+++ b/Engine/Plugins/Runtime/InstancedActors/Source/InstancedActors/Public/InstancedActorsSettingsTypes.h
@@ -139,7 +139,7 @@ struct INSTANCEDACTORS_API FInstancedActorsSettings : public FTableRowBase
 	UPROPERTY(EditAnywhere, BlueprintReadWrite, Meta = (EditCondition = "bOverride_bDisableAutoDistanceCulling"), Category=InstancedActors)
 	bool bDisableAutoDistanceCulling = false;
 
-	/** Fix for PLAY-11011. If false, collision will not be managed by the mass LODs for this instance **/
+	/** If false, collision will not be managed by the mass LODs for this instance **/
 	UPROPERTY(EditAnywhere, BlueprintReadWrite, Meta = (EditCondition = "bOverride_bControlPhysicsState"), Category=InstancedActors)
 	bool bControlPhysicsState = true;	
 
