Fixed AInstancedActorsManager::PreRegisterAllComponents() missing call to Super::PreRegisterAllComponents(). (in preparation for a future engine change that requires a call to Super).

diff --git a/Engine/Plugins/Runtime/InstancedActors/Source/InstancedActors/Private/InstancedActorsManager.cpp b/Engine/Plugins/Runtime/InstancedActors/Source/InstancedActors/Private/InstancedActorsManager.cpp
--- a/Engine/Plugins/Runtime/InstancedActors/Source/InstancedActors/Private/InstancedActorsManager.cpp
+++ b/Engine/Plugins/Runtime/InstancedActors/Source/InstancedActors/Private/InstancedActorsManager.cpp
@@ -738,6 +738,8 @@ UInstancedActorsData& AInstancedActorsManager::GetOrCreateActorInstanceData(TSub
 
 void AInstancedActorsManager::PreRegisterAllComponents()
 {
+	Super::PreRegisterAllComponents();
+
 	for (UInstancedActorsData* InstanceData : PerActorClassInstanceData)
 	{
 		// Modify the ISMCs here since the components might have been serialized before this change.
