[InstancedActors] Fixed IA configuration issue resulting in standalone games using "client" elements, which is not by design.

diff --git a/Engine/Plugins/Runtime/InstancedActors/Source/InstancedActors/Private/ClientInstancedActorsSpawnerSubsystem.cpp b/Engine/Plugins/Runtime/InstancedActors/Source/InstancedActors/Private/ClientInstancedActorsSpawnerSubsystem.cpp
--- a/Engine/Plugins/Runtime/InstancedActors/Source/InstancedActors/Private/ClientInstancedActorsSpawnerSubsystem.cpp
+++ b/Engine/Plugins/Runtime/InstancedActors/Source/InstancedActors/Private/ClientInstancedActorsSpawnerSubsystem.cpp
@@ -28,7 +28,7 @@ bool UClientInstancedActorsSpawnerSubsystem::ShouldCreateSubsystem(UObject* Oute
 	// @todo Add support for non-replay NM_Standalone where we should use UServerInstancedActorsSpawnerSubsystem for 
 	// authoritative actor spawning.
 	UWorld* World = Cast<UWorld>(Outer);
-	return (World != nullptr && World->GetNetMode() != NM_DedicatedServer);
+	return (World != nullptr && World->GetNetMode() == NM_Client);
 }
 
 ESpawnRequestStatus UClientInstancedActorsSpawnerSubsystem::SpawnActor(FConstStructView SpawnRequestView, TObjectPtr<AActor>& OutSpawnedActor, FActorSpawnParameters& InOutSpawnParameters) const
diff --git a/Engine/Plugins/Runtime/InstancedActors/Source/InstancedActors/Private/InstancedActorsTypes.cpp b/Engine/Plugins/Runtime/InstancedActors/Source/InstancedActors/Private/InstancedActorsTypes.cpp
--- a/Engine/Plugins/Runtime/InstancedActors/Source/InstancedActors/Private/InstancedActorsTypes.cpp
+++ b/Engine/Plugins/Runtime/InstancedActors/Source/InstancedActors/Private/InstancedActorsTypes.cpp
@@ -16,12 +16,11 @@ namespace UE::InstancedActors::Utils
 	{
 		// @todo Add support for non-replay NM_Standalone where we should use UServerInstancedActorsSpawnerSubsystem for 
 		// authoritative actor spawning.
-		if (World.GetNetMode() == NM_DedicatedServer)
+		if (World.GetNetMode() == NM_Client)
 		{
-			return GET_INSTANCEDACTORS_CONFIG_VALUE(GetServerActorSpawnerSubsystemClass());
-
+			return GET_INSTANCEDACTORS_CONFIG_VALUE(GetClientActorSpawnerSubsystemClass());
 		}
-		return GET_INSTANCEDACTORS_CONFIG_VALUE(GetClientActorSpawnerSubsystemClass());
+		return GET_INSTANCEDACTORS_CONFIG_VALUE(GetServerActorSpawnerSubsystemClass());
 	}
 
 	UServerInstancedActorsSpawnerSubsystem* GetServerInstancedActorsSpawnerSubsystem(const UWorld& World)
@@ -40,12 +39,11 @@ namespace UE::InstancedActors::Utils
 
 	UMassActorSpawnerSubsystem* GetActorSpawnerSubsystem(const UWorld& World)
 	{
-		if (World.GetNetMode() == NM_DedicatedServer)
+		if (World.GetNetMode() == NM_Client)
 		{
-			return GetServerInstancedActorsSpawnerSubsystem(World);
-
+			return GetClientInstancedActorsSpawnerSubsystem(World);
 		}
-		return GetClientInstancedActorsSpawnerSubsystem(World);
+		return GetServerInstancedActorsSpawnerSubsystem(World);
 	}
 
 	UInstancedActorsSubsystem* GetInstancedActorsSubsystem(const UWorld& World)
diff --git a/Engine/Plugins/Runtime/InstancedActors/Source/InstancedActors/Private/InstancedActorsVisualizationTrait.cpp b/Engine/Plugins/Runtime/InstancedActors/Source/InstancedActors/Private/InstancedActorsVisualizationTrait.cpp
--- a/Engine/Plugins/Runtime/InstancedActors/Source/InstancedActors/Private/InstancedActorsVisualizationTrait.cpp
+++ b/Engine/Plugins/Runtime/InstancedActors/Source/InstancedActors/Private/InstancedActorsVisualizationTrait.cpp
@@ -33,8 +33,8 @@ void UInstancedActorsVisualizationTrait::InitializeFromInstanceData(UInstancedAc
 
 	HighResTemplateActor = InstanceData->ActorClass;
 
-	const bool bIsDedicatedServer = InInstanceData.GetManagerChecked().IsNetMode(NM_DedicatedServer);
-	if (!bIsDedicatedServer)
+	const bool bIsClient = InInstanceData.GetManagerChecked().IsNetMode(NM_Client);
+	if (bIsClient)
 	{
 		// Don't attempt to spawn actors natively on clients. Instead, rely on bForceActorRepresentationForExternalActors to 
 		// switch to Actor representation once replicated actors are set explicitly in UInstancedActorsData::SetReplicatedActor 
@@ -54,7 +54,7 @@ void UInstancedActorsVisualizationTrait::InitializeFromInstanceData(UInstancedAc
 		ensure(InstanceData->GetDefaultVisualizationChecked().ISMComponents.IsEmpty());
 		Params.LODRepresentation[EMassLOD::Low] = EMassRepresentationType::None;
 
-		if (!bIsDedicatedServer)
+		if (bIsClient)
 		{
 			// Don't attempt to switch to ISMC representation on clients for no mesh classes (which would otherwise crash)
 			// Let the server spawn these actors and replicate them to clients.
diff --git a/Engine/Plugins/Runtime/InstancedActors/Source/InstancedActors/Private/ServerInstancedActorsSpawnerSubsystem.cpp b/Engine/Plugins/Runtime/InstancedActors/Source/InstancedActors/Private/ServerInstancedActorsSpawnerSubsystem.cpp
--- a/Engine/Plugins/Runtime/InstancedActors/Source/InstancedActors/Private/ServerInstancedActorsSpawnerSubsystem.cpp
+++ b/Engine/Plugins/Runtime/InstancedActors/Source/InstancedActors/Private/ServerInstancedActorsSpawnerSubsystem.cpp
@@ -20,7 +20,7 @@ bool UServerInstancedActorsSpawnerSubsystem::ShouldCreateSubsystem(UObject* Oute
 	// @todo Add support for non-replay NM_Standalone where we should use UServerInstancedActorsSpawnerSubsystem for 
 	// authoritative actor spawning.
 	UWorld* World = Cast<UWorld>(Outer);
-	return (World != nullptr && World->GetNetMode() == NM_DedicatedServer);
+	return (World != nullptr && World->GetNetMode() != NM_Client);
 }
 
 bool UServerInstancedActorsSpawnerSubsystem::ReleaseActorToPool(AActor* Actor)
@@ -50,7 +50,7 @@ ESpawnRequestStatus UServerInstancedActorsSpawnerSubsystem::SpawnActor(FConstStr
 
 	UWorld* World = GetWorld();
 	check(World);
-	check(World->GetNetMode() == NM_DedicatedServer);
+	check(World->GetNetMode() != NM_Client);
 
 	const FMassActorSpawnRequest& SpawnRequest = SpawnRequestView.Get<const FMassActorSpawnRequest>();
 	UInstancedActorsData* InstanceData = UInstancedActorsData::GetInstanceDataForEntity(*EntityManager, SpawnRequest.MassAgent);
