[InstancedActors] Massaged InstancedActor's ExemplarActor handling to avoid `checkSlow(KeyHash == KeyFuncs::GetKeyHash(Key));` in TMap operations

diff --git a/Engine/Plugins/Runtime/InstancedActors/Source/InstancedActors/Private/InstancedActorsSubsystem.cpp b/Engine/Plugins/Runtime/InstancedActors/Source/InstancedActors/Private/InstancedActorsSubsystem.cpp
--- a/Engine/Plugins/Runtime/InstancedActors/Source/InstancedActors/Private/InstancedActorsSubsystem.cpp
+++ b/Engine/Plugins/Runtime/InstancedActors/Source/InstancedActors/Private/InstancedActorsSubsystem.cpp
@@ -623,13 +623,14 @@ TSharedRef<UE::InstancedActors::FExemplarActorData> UInstancedActorsSubsystem::G
 {
 	TRACE_CPUPROFILER_EVENT_SCOPE(UInstancedActorsSubsystem GetOrCreateExemplarActor);
 
-	const UClass* const ActorClassPtr = ActorClass.Get();
+	UClass* ActorClassPtr = ActorClass.Get();
 	check(ActorClassPtr);
 
 	// Return existing?
-	const uint32 ActorClassHash = GetTypeHash(TObjectKey<const UClass>(ActorClassPtr));
+	const TObjectKey<const UClass> ActorClassKey(ActorClassPtr);
+	const uint32 ActorClassHash = GetTypeHash(ActorClassKey);
 	
-	if (const TWeakPtr<UE::InstancedActors::FExemplarActorData>* CachedExemplarActorDataPtr = ExemplarActors.FindByHash(ActorClassHash, ActorClassPtr))
+	if (const TWeakPtr<UE::InstancedActors::FExemplarActorData>* CachedExemplarActorDataPtr = ExemplarActors.FindByHash(ActorClassHash, ActorClassKey))
 	{
 		// This can fail in editor with undo/redo in the mix.
 		TSharedPtr<UE::InstancedActors::FExemplarActorData> CachedExemplarActorData = CachedExemplarActorDataPtr->Pin();
@@ -640,7 +641,7 @@ TSharedRef<UE::InstancedActors::FExemplarActorData> UInstancedActorsSubsystem::G
 		else
 		{
 			// The examplar is not valid, we'll remove it and then re-create it below.
-			ExemplarActors.RemoveByHash(ActorClassHash, ActorClassPtr);
+			ExemplarActors.RemoveByHash(ActorClassHash, ActorClassKey);
 		}
 	}
 
@@ -684,12 +685,12 @@ TSharedRef<UE::InstancedActors::FExemplarActorData> UInstancedActorsSubsystem::G
 	SpawnParameters.SpawnCollisionHandlingOverride = ESpawnActorCollisionHandlingMethod::AlwaysSpawn;
 
 	check(ExemplarActorWorld);
-	AActor* NewExemplarActor = ExemplarActorWorld->SpawnActor(ActorClass, /*Transform*/nullptr, SpawnParameters);
+	AActor* NewExemplarActor = ExemplarActorWorld->SpawnActor(ActorClassPtr, /*Transform*/nullptr, SpawnParameters);
 	check(NewExemplarActor);
 
 	// Cache for subsequent calls
 	TSharedPtr<UE::InstancedActors::FExemplarActorData> NewExemplarActorDataPtr{new UE::InstancedActors::FExemplarActorData{*NewExemplarActor, *this}};
-	ExemplarActors.AddByHash(ActorClassHash, ActorClassPtr, NewExemplarActorDataPtr);
+	ExemplarActors.AddByHash(ActorClassHash, ActorClassKey, NewExemplarActorDataPtr);
 
 	return NewExemplarActorDataPtr.ToSharedRef();
 }
@@ -698,11 +699,12 @@ void UInstancedActorsSubsystem::UnregisterExemplarActorClass(TSubclassOf<AActor>
 {
 	const UClass* const ActorClassPtr = ActorClass.Get();
 	check(ActorClassPtr);
+	const TObjectKey<const UClass> ActorClassKey(ActorClassPtr);
 	
-	const uint32 ActorClassHash = GetTypeHash(TObjectKey<const UClass>(ActorClassPtr));	
+	const uint32 ActorClassHash = GetTypeHash(ActorClassKey);	
 	if (const TWeakPtr<UE::InstancedActors::FExemplarActorData>* CachedExemplarActorDataPtr = ExemplarActors.FindByHash(ActorClassHash, ActorClassPtr))
 	{
-		ExemplarActors.RemoveByHash(ActorClassHash, ActorClassPtr);
+		ExemplarActors.RemoveByHash(ActorClassHash, ActorClassKey);
 	}
 }
 
