[InstancedActors] Forcing server-side spawned actors to be replicated - it's required so that the actors show up on the client.

diff --git a/Engine/Plugins/Runtime/InstancedActors/Source/InstancedActors/Private/ServerInstancedActorsSpawnerSubsystem.cpp b/Engine/Plugins/Runtime/InstancedActors/Source/InstancedActors/Private/ServerInstancedActorsSpawnerSubsystem.cpp
--- a/Engine/Plugins/Runtime/InstancedActors/Source/InstancedActors/Private/ServerInstancedActorsSpawnerSubsystem.cpp
+++ b/Engine/Plugins/Runtime/InstancedActors/Source/InstancedActors/Private/ServerInstancedActorsSpawnerSubsystem.cpp
@@ -92,6 +92,10 @@ ESpawnRequestStatus UServerInstancedActorsSpawnerSubsystem::SpawnActor(FConstStr
 		{
 			// No existing UInstancedActorsComponent class or subclass, add a new UInstancedActorsComponent
 			InstancedActorComponent = NewObject<UInstancedActorsComponent>(OutSpawnedActor);
+			if (OutSpawnedActor->GetIsReplicated() == false)
+			{
+				OutSpawnedActor->SetReplicates(true);
+			}
 			InstancedActorComponent->SetIsReplicated(true);
 			InstancedActorComponent->RegisterComponent();
 		}
