[InstancedActors] InstancedActors changes to support generic usecase.

With this change IA can be used in any project both in standalone mode as well as in multiplayer. Works both in world-partitioned worlds as well as in regular maps.

diff --git a/Engine/Plugins/Runtime/InstancedActors/Source/InstancedActors/Private/InstancedActorsData.cpp b/Engine/Plugins/Runtime/InstancedActors/Source/InstancedActors/Private/InstancedActorsData.cpp
--- a/Engine/Plugins/Runtime/InstancedActors/Source/InstancedActors/Private/InstancedActorsData.cpp
+++ b/Engine/Plugins/Runtime/InstancedActors/Source/InstancedActors/Private/InstancedActorsData.cpp
@@ -20,6 +20,7 @@
 #if UE_WITH_IRIS
 #include "Iris/ReplicationSystem/ReplicationFragmentUtil.h"
 #endif
+#include "InstancedActorsVisualizationProcessor.h"
 #include "MassActorSubsystem.h"
 #include "MassCommonFragments.h"
 #include "MassEntityConfigAsset.h"
@@ -313,7 +314,7 @@ void UInstancedActorsData::ModifyEntityTemplate(FMassEntityTemplateData& Modifie
 {
 	ModifiedTemplate.RemoveTag<FMassDistanceLODProcessorTag>();
 	ModifiedTemplate.RemoveTag<FMassCollectDistanceLODViewerInfoTag>();
-	ModifiedTemplate.RemoveTag<FMassVisualizationProcessorTag>();
+	ModifiedTemplate.RemoveTag<FInstancedActorsVisualizationProcessorTag>();
 	// not needed really, since we don't add it in any of the traits but leaving here for the reference
 	// ModifiedTemplate.RemoveTag<FMassStationaryISMSwitcherProcessorTag>();
 }
diff --git a/Engine/Plugins/Runtime/InstancedActors/Source/InstancedActors/Private/InstancedActorsRepresentationActorManagement.cpp b/Engine/Plugins/Runtime/InstancedActors/Source/InstancedActors/Private/InstancedActorsRepresentationActorManagement.cpp
--- a/Engine/Plugins/Runtime/InstancedActors/Source/InstancedActors/Private/InstancedActorsRepresentationActorManagement.cpp
+++ b/Engine/Plugins/Runtime/InstancedActors/Source/InstancedActors/Private/InstancedActorsRepresentationActorManagement.cpp
@@ -4,6 +4,7 @@
 #include "InstancedActorsData.h"
 #include "InstancedActorsSubsystem.h"
 #include "InstancedActorsSettingsTypes.h"
+#include "InstancedActorsVisualizationProcessor.h"
 #include "Delegates/Delegate.h"
 #include "MassEntitySubsystem.h"
 #include "MassActorSpawnerSubsystem.h"
@@ -201,7 +202,7 @@ AActor* UInstancedActorsRepresentationActorManagement::FindOrInstantlySpawnActor
 		// note that we do need to make the change synchronously since due to EntityManager.IsProcessing() == false
 		// any command we issue here might get called after LOD and Visualization processing, that could override the
 		// values we've just set
-		FMassTagBitSet TagsToRemove = UE::Mass::Utils::ConstructTagBitSet<EMassCommandCheckTime::CompileTimeCheck, FMassStationaryISMSwitcherProcessorTag, FMassVisualizationProcessorTag>();
+		FMassTagBitSet TagsToRemove = UE::Mass::Utils::ConstructTagBitSet<EMassCommandCheckTime::CompileTimeCheck, FMassStationaryISMSwitcherProcessorTag, FInstancedActorsVisualizationProcessorTag>();
 		EntityManager.RemoveCompositionFromEntity(EntityView.GetEntity(), FMassArchetypeCompositionDescriptor(MoveTemp(TagsToRemove)));
 	}
 
diff --git a/Engine/Plugins/Runtime/InstancedActors/Source/InstancedActors/Private/InstancedActorsStationaryLODBatchProcessor.cpp b/Engine/Plugins/Runtime/InstancedActors/Source/InstancedActors/Private/InstancedActorsStationaryLODBatchProcessor.cpp
--- a/Engine/Plugins/Runtime/InstancedActors/Source/InstancedActors/Private/InstancedActorsStationaryLODBatchProcessor.cpp
+++ b/Engine/Plugins/Runtime/InstancedActors/Source/InstancedActors/Private/InstancedActorsStationaryLODBatchProcessor.cpp
@@ -9,6 +9,7 @@
 #include "InstancedActorsData.h"
 #include "InstancedActorsDebug.h"
 #include "InstancedActorsSubsystem.h"
+#include "InstancedActorsVisualizationProcessor.h"
 #include "Components/InstancedStaticMeshComponent.h"
 #include "StaticMeshResources.h"
 
@@ -131,8 +132,8 @@ void UInstancedActorsStationaryLODBatchProcessor::Execute(FMassEntityManager& En
 {
 	SCOPE_CYCLE_COUNTER(STAT_InstancedActorsStationaryLODBatchProcessor_Execute);
 
-	using FAddRelevantTagsCommand = FMassCommandAddTags<FMassInActiveSmartObjectsRangeTag, FMassDistanceLODProcessorTag, FMassCollectDistanceLODViewerInfoTag, FMassStationaryISMSwitcherProcessorTag, FMassVisualizationProcessorTag>;
-	using FRemoveRelevantTagsCommand = FMassCommandRemoveTags<FMassInActiveSmartObjectsRangeTag, FMassDistanceLODProcessorTag, FMassCollectDistanceLODViewerInfoTag, FMassStationaryISMSwitcherProcessorTag, FMassVisualizationProcessorTag>;
+	using FAddRelevantTagsCommand = FMassCommandAddTags<FMassInActiveSmartObjectsRangeTag, FMassDistanceLODProcessorTag, FMassCollectDistanceLODViewerInfoTag, FMassStationaryISMSwitcherProcessorTag, FInstancedActorsVisualizationProcessorTag>;
+	using FRemoveRelevantTagsCommand = FMassCommandRemoveTags<FMassInActiveSmartObjectsRangeTag, FMassDistanceLODProcessorTag, FMassCollectDistanceLODViewerInfoTag, FMassStationaryISMSwitcherProcessorTag, FInstancedActorsVisualizationProcessorTag>;
 
 	// some of the code below assumes EInstancedActorsBulkLOD::Detailed == 0, we need to verify that's the case. If not the code below needs updating.
 	static_assert((uint8)EInstancedActorsBulkLOD::Detailed == 0, "Code below relies on the assumptions. Needs to be updated if the assumption is broken");
@@ -172,7 +173,8 @@ void UInstancedActorsStationaryLODBatchProcessor::Execute(FMassEntityManager& En
 			// We need to filter out the latter. 
 			// Note that we rely on UMassSubsystem::bUsePlayerPawnLocationInsteadOfCamera being true here. Without it there's no 
 			// reliable way to differentiate the cases.
-			checkSlow(LODSubsystem.IsUsingPlayerPawnLocationInsteadOfCamera());
+			UE_CLOG(LODSubsystem.IsUsingPlayerPawnLocationInsteadOfCamera() == false
+				, LogInstancedActors, Warning, TEXT("Using Player's camera location for instanced actors LOD calculations - this can skew the LOD calculations in non-FPP games."));
 			if (APlayerController* ViewerAsPlayerController = Viewers[ViewerIndex].GetPlayerController())
 			{
 				if (ViewerAsPlayerController->GetPawn() == nullptr)
@@ -442,7 +444,7 @@ void UInstancedActorsStationaryLODBatchProcessor::Execute(FMassEntityManager& En
 					{
 						// It's possible that we've only just switched to Non-Detailed this frame, the tag removal to prevent regular processing wouldn't have
 						// occurred yet and we would have performed a representation update this frame already.
-						if (!Context.DoesArchetypeHaveTag<FMassVisualizationProcessorTag>())
+						if (!Context.DoesArchetypeHaveTag<FInstancedActorsVisualizationProcessorTag>())
 						{
 							FMassRepresentationUpdateParams Params;
 							Params.bTestCollisionAvailibilityForActorVisualization = false;
diff --git a/Engine/Plugins/Runtime/InstancedActors/Source/InstancedActors/Private/InstancedActorsVisualizationProcessor.cpp b/Engine/Plugins/Runtime/InstancedActors/Source/InstancedActors/Private/InstancedActorsVisualizationProcessor.cpp
new file mode 100644
--- /dev/null
+++ b/Engine/Plugins/Runtime/InstancedActors/Source/InstancedActors/Private/InstancedActorsVisualizationProcessor.cpp
@@ -0,0 +1,26 @@
+// Copyright Epic Games, Inc. All Rights Reserved.
+
+#include "InstancedActorsVisualizationProcessor.h"
+#include "MassStationaryISMSwitcherProcessor.h"
+
+
+UInstancedActorsVisualizationProcessor::UInstancedActorsVisualizationProcessor()
+{
+	bAutoRegisterWithProcessingPhases = true;
+	ExecutionFlags = static_cast<int32>(EProcessorExecutionFlags::AllNetModes);
+	// This processor needs to be executed before UMassStationaryISMSwitcherProcessors since that's the processor
+	// responsible for executing what UInstancedActorsVisualizationProcessor calculates.
+	// Missing this dependency would result in client-side one-frame representation absence when switching
+	// from actor representation back to ISM.
+	ExecutionOrder.ExecuteBefore.Add(UMassStationaryISMSwitcherProcessor::StaticClass()->GetFName());
+
+	UpdateParams.bTestCollisionAvailibilityForActorVisualization = false;
+}
+
+void UInstancedActorsVisualizationProcessor::ConfigureQueries() 
+{
+	Super::ConfigureQueries();
+
+	EntityQuery.ClearTagRequirements(FMassTagBitSet(*FMassVisualizationProcessorTag::StaticStruct()));
+	EntityQuery.AddTagRequirement<FInstancedActorsVisualizationProcessorTag>(EMassFragmentPresence::All);
+}
\ No newline at end of file
diff --git a/Engine/Plugins/Runtime/InstancedActors/Source/InstancedActors/Private/InstancedActorsVisualizationTrait.cpp b/Engine/Plugins/Runtime/InstancedActors/Source/InstancedActors/Private/InstancedActorsVisualizationTrait.cpp
--- a/Engine/Plugins/Runtime/InstancedActors/Source/InstancedActors/Private/InstancedActorsVisualizationTrait.cpp
+++ b/Engine/Plugins/Runtime/InstancedActors/Source/InstancedActors/Private/InstancedActorsVisualizationTrait.cpp
@@ -7,6 +7,7 @@
 #include "InstancedActorsRepresentationSubsystem.h"
 #include "InstancedActorsSettingsTypes.h"
 #include "InstancedActorsTypes.h"
+#include "InstancedActorsVisualizationProcessor.h"
 #include "MassEntityTemplateRegistry.h"
 #include "MassCommonFragments.h"
 #include "MassActorSubsystem.h"
@@ -80,6 +81,9 @@ void UInstancedActorsVisualizationTrait::BuildTemplate(FMassEntityTemplateBuildC
 
 	Super::BuildTemplate(BuildContext, World);
 
+	// we need IAs to be processed by a dedicated visualization processor, configured a bit differently than the default one.
+	BuildContext.RemoveTag<FMassVisualizationProcessorTag>();
+
 	FMassEntityManager& EntityManager = UE::Mass::Utils::GetEntityManagerChecked(World);
 
 	BuildContext.AddFragment<FTransformFragment>();
diff --git a/Engine/Plugins/Runtime/InstancedActors/Source/InstancedActors/Public/InstancedActorsVisualizationProcessor.h b/Engine/Plugins/Runtime/InstancedActors/Source/InstancedActors/Public/InstancedActorsVisualizationProcessor.h
new file mode 100644
--- /dev/null
+++ b/Engine/Plugins/Runtime/InstancedActors/Source/InstancedActors/Public/InstancedActorsVisualizationProcessor.h
@@ -0,0 +1,27 @@
+// Copyright Epic Games, Inc. All Rights Reserved.
+
+#pragma once
+
+#include "MassRepresentationProcessor.h"
+#include "InstancedActorsVisualizationProcessor.generated.h"
+
+
+/**
+ * Tag required by Instanced Actors Visualization Processor to process given archetype. Removing the tag allows to support
+ * disabling of processing for individual entities of given archetype.
+ */
+USTRUCT()
+struct INSTANCEDACTORS_API FInstancedActorsVisualizationProcessorTag : public FMassTag
+{
+	GENERATED_BODY();
+};
+
+UCLASS()
+class UInstancedActorsVisualizationProcessor : public UMassVisualizationProcessor
+{
+	GENERATED_BODY()
+
+protected:
+	UInstancedActorsVisualizationProcessor();
+	virtual void ConfigureQueries() override;
+};
