[InstancedActors] Fixed InstancedActorsVisualizationTrait failing checks during traits validation in editor.

diff --git a/Engine/Plugins/Runtime/InstancedActors/Source/InstancedActors/Private/InstancedActorsVisualizationTrait.cpp b/Engine/Plugins/Runtime/InstancedActors/Source/InstancedActors/Private/InstancedActorsVisualizationTrait.cpp
--- a/Engine/Plugins/Runtime/InstancedActors/Source/InstancedActors/Private/InstancedActorsVisualizationTrait.cpp
+++ b/Engine/Plugins/Runtime/InstancedActors/Source/InstancedActors/Private/InstancedActorsVisualizationTrait.cpp
@@ -77,7 +77,7 @@ void UInstancedActorsVisualizationTrait::InitializeFromInstanceData(UInstancedAc
 
 void UInstancedActorsVisualizationTrait::BuildTemplate(FMassEntityTemplateBuildContext& BuildContext, const UWorld& World) const
 {
-	check(InstanceData.IsValid());
+	check(InstanceData.IsValid() || BuildContext.IsInspectingData());
 
 	Super::BuildTemplate(BuildContext, World);
 
@@ -108,7 +108,10 @@ void UInstancedActorsVisualizationTrait::BuildTemplate(FMassEntityTemplateBuildC
 		// InstanceData while preserving the "runtime" value, which will mess up newly spawned entities.
 		AsShared->BulkLOD = EInstancedActorsBulkLOD::MAX;
 
-		InstanceData->SetSharedInstancedActorDataStruct(SubsystemFragment);
+		if (BuildContext.IsInspectingData() == false)
+		{
+			InstanceData->SetSharedInstancedActorDataStruct(SubsystemFragment);
+		}
 	}
 	// not adding SubsystemFragment do BuildContext on purpose, we temporarily use shared fragments to store IAD information.
 	// To be moved to InstancedActorSubsystem in the future
@@ -119,19 +122,22 @@ void UInstancedActorsVisualizationTrait::BuildTemplate(FMassEntityTemplateBuildC
 	FInstancedActorsFragment& InstancedActorFragment = BuildContext.AddFragment_GetRef<FInstancedActorsFragment>();
 	InstancedActorFragment.InstanceData = InstanceData;
 
-	// @todo Implement version of AddVisualDescWithISMComponent that supports multiple ISMCs and use that here
-	if (ensure(InstanceData.IsValid()))
+	if (BuildContext.IsInspectingData() == false)
 	{
-		FMassRepresentationFragment* RepresentationFragment = BuildContext.GetFragment<FMassRepresentationFragment>();
-		if (ensureMsgf(RepresentationFragment, TEXT("Configuration error, we always expect to have a FMassRepresentationFragment instance at this point")))
+		// @todo Implement version of AddVisualDescWithISMComponent that supports multiple ISMCs and use that here
+		if (ensure(InstanceData.IsValid()))
 		{
-			RepresentationFragment->StaticMeshDescHandle = InstanceData->GetDefaultVisualizationChecked().MassStaticMeshDescHandle;
-
-			if (RepresentationFragment->LowResTemplateActorIndex == INDEX_NONE)
+			FMassRepresentationFragment* RepresentationFragment = BuildContext.GetFragment<FMassRepresentationFragment>();
+			if (ensureMsgf(RepresentationFragment, TEXT("Configuration error, we always expect to have a FMassRepresentationFragment instance at this point")))
 			{
-				// if there's no "low res actor" we reuse the high-res one, otherwise we risk the visualization actor getting 
-				// removed when switching from EMassLOD::High down to EMassLOD::Medium
-				RepresentationFragment->LowResTemplateActorIndex = RepresentationFragment->HighResTemplateActorIndex;
+				RepresentationFragment->StaticMeshDescHandle = InstanceData->GetDefaultVisualizationChecked().MassStaticMeshDescHandle;
+
+				if (RepresentationFragment->LowResTemplateActorIndex == INDEX_NONE)
+				{
+					// if there's no "low res actor" we reuse the high-res one, otherwise we risk the visualization actor getting 
+					// removed when switching from EMassLOD::High down to EMassLOD::Medium
+					RepresentationFragment->LowResTemplateActorIndex = RepresentationFragment->HighResTemplateActorIndex;
+				}
 			}
 		}
 	}
