[InstancedActors] Added two InstancedActors-specific Mass commands for switching between Detailed and Batched LOD. This approach makes keeping added/removed tag types in sync while at the same time lets project-specific code affect which tags will be used.

diff --git a/Engine/Plugins/Runtime/InstancedActors/Source/InstancedActors/Private/InstancedActorsCommands.cpp b/Engine/Plugins/Runtime/InstancedActors/Source/InstancedActors/Private/InstancedActorsCommands.cpp
new file mode 100644
--- /dev/null
+++ b/Engine/Plugins/Runtime/InstancedActors/Source/InstancedActors/Private/InstancedActorsCommands.cpp
@@ -0,0 +1,58 @@
+// Copyright Epic Games, Inc. All Rights Reserved.
+
+#include "InstancedActorsCommands.h"
+#include "MassEntityTypes.h"
+#include "MassLODFragments.h"
+#include "MassDistanceLODProcessor.h"
+#include "MassStationaryISMSwitcherProcessor.h"
+#include "InstancedActorsVisualizationProcessor.h"
+
+#ifndef INSTANCEDACTORS_AS_SMARTOBJECTS
+#define INSTANCEDACTORS_AS_SMARTOBJECTS 1
+#endif // INSTANCEDACTORS_AS_SMARTOBJECTS
+
+#if INSTANCEDACTORS_AS_SMARTOBJECTS
+#include "MassSmartObjectRegistration.h"
+#endif // INSTANCEDACTORS_AS_SMARTOBJECTS
+
+#if CSV_PROFILER_STATS || WITH_INSTANCEDACTORS_DEBUG
+#	define DEBUG_NAME(Name) , FName(TEXT(Name))
+#else
+#	define DEBUG_NAME(Name)
+#endif // CSV_PROFILER_STATS || WITH_MASSENTITY_DEBUG
+
+namespace UE::InstancedActors
+{	
+	FMassTagBitSet& GetDetailedLODTags()
+	{
+		static FMassTagBitSet DetailedLODTags = UE::Mass::Utils::ConstructTagBitSet<EMassCommandCheckTime::CompileTimeCheck
+			, FMassDistanceLODProcessorTag				// UMassDistanceLODProcessor requirement
+			, FMassCollectDistanceLODViewerInfoTag		// UMassLODDistanceCollectorProcessor requirement
+			, FMassStationaryISMSwitcherProcessorTag	// UMassStationaryISMSwitcherProcessor requirement
+			, FInstancedActorsVisualizationProcessorTag	// UInstancedActorsVisualizationProcessor requirement
+#if INSTANCEDACTORS_AS_SMARTOBJECTS
+			, FMassInActiveSmartObjectsRangeTag
+#endif // INSTANCEDACTORS_AS_SMARTOBJECTS
+		>();
+
+		return DetailedLODTags;
+	}
+
+	FEnableDetailedLODCommand::FEnableDetailedLODCommand()
+		: Super(EMassCommandOperationType::Add
+			, GetDetailedLODTags()
+			, FMassTagBitSet{}
+			DEBUG_NAME("DetailedLODEnable"))
+	{
+	}
+
+	FEnableBatchLODCommand::FEnableBatchLODCommand()
+		: Super(EMassCommandOperationType::Remove
+			, FMassTagBitSet{}
+			, GetDetailedLODTags()
+			DEBUG_NAME("BatchLODEnable"))
+	{
+	}
+} // UE::InstancedActors
+
+#undef DEBUG_NAME
diff --git a/Engine/Plugins/Runtime/InstancedActors/Source/InstancedActors/Private/InstancedActorsData.cpp b/Engine/Plugins/Runtime/InstancedActors/Source/InstancedActors/Private/InstancedActorsData.cpp
--- a/Engine/Plugins/Runtime/InstancedActors/Source/InstancedActors/Private/InstancedActorsData.cpp
+++ b/Engine/Plugins/Runtime/InstancedActors/Source/InstancedActors/Private/InstancedActorsData.cpp
@@ -11,6 +11,7 @@
 #include "InstancedActorsVisualizationTrait.h"
 #include "InstancedActorsSettingsTypes.h"
 #include "InstancedActorsSettings.h"
+#include "InstancedActorsCommands.h"
 #include "UObject/ObjectSaveContext.h"
 #include "Algo/Count.h"
 #include "Algo/NoneOf.h"
@@ -312,11 +313,9 @@ void UInstancedActorsData::CreateEntityTemplate(const AActor& ExemplarActor)
 
 void UInstancedActorsData::ModifyEntityTemplate(FMassEntityTemplateData& ModifiedTemplate, const AActor&)
 {
-	ModifiedTemplate.RemoveTag<FMassDistanceLODProcessorTag>();
-	ModifiedTemplate.RemoveTag<FMassCollectDistanceLODViewerInfoTag>();
-	ModifiedTemplate.RemoveTag<FInstancedActorsVisualizationProcessorTag>();
-	// not needed really, since we don't add it in any of the traits but leaving here for the reference
-	// ModifiedTemplate.RemoveTag<FMassStationaryISMSwitcherProcessorTag>();
+	// NOTE: Removes some tags that are added by default by the traits that we add to the original template,
+	// but are toggled by the UInstancedActorsStationaryLODBatchProcessor and shouldn't be active when the Mass entities are created.
+	ModifiedTemplate.GetMutableTags().Remove(UE::InstancedActors::GetDetailedLODTags());
 }
 
 void UInstancedActorsData::ReleaseEntityTemplate()
diff --git a/Engine/Plugins/Runtime/InstancedActors/Source/InstancedActors/Private/InstancedActorsStationaryLODBatchProcessor.cpp b/Engine/Plugins/Runtime/InstancedActors/Source/InstancedActors/Private/InstancedActorsStationaryLODBatchProcessor.cpp
--- a/Engine/Plugins/Runtime/InstancedActors/Source/InstancedActors/Private/InstancedActorsStationaryLODBatchProcessor.cpp
+++ b/Engine/Plugins/Runtime/InstancedActors/Source/InstancedActors/Private/InstancedActorsStationaryLODBatchProcessor.cpp
@@ -9,22 +9,20 @@
 #include "InstancedActorsData.h"
 #include "InstancedActorsDebug.h"
 #include "InstancedActorsSubsystem.h"
+#include "InstancedActorsCommands.h"
 #include "InstancedActorsVisualizationProcessor.h"
 #include "Components/InstancedStaticMeshComponent.h"
 #include "StaticMeshResources.h"
 
 #include "MassActorSubsystem.h"
 #include "MassCommands.h"
-#include "MassDistanceLODProcessor.h"
 #include "MassExecutionContext.h"
-#include "MassLODFragments.h"
 #include "MassLODSubsystem.h"
 #include "MassLODTypes.h"
 #include "MassRepresentationFragments.h"
 #include "MassRepresentationProcessor.h"
 #include "MassRepresentationTypes.h"
 #include "MassSignalSubsystem.h"
-#include "MassSmartObjectRegistration.h"
 #include "MassStationaryISMSwitcherProcessor.h"
 
 
@@ -132,9 +130,6 @@ void UInstancedActorsStationaryLODBatchProcessor::Execute(FMassEntityManager& En
 {
 	SCOPE_CYCLE_COUNTER(STAT_InstancedActorsStationaryLODBatchProcessor_Execute);
 
-	using FAddRelevantTagsCommand = FMassCommandAddTags<FMassInActiveSmartObjectsRangeTag, FMassDistanceLODProcessorTag, FMassCollectDistanceLODViewerInfoTag, FMassStationaryISMSwitcherProcessorTag, FInstancedActorsVisualizationProcessorTag>;
-	using FRemoveRelevantTagsCommand = FMassCommandRemoveTags<FMassInActiveSmartObjectsRangeTag, FMassDistanceLODProcessorTag, FMassCollectDistanceLODViewerInfoTag, FMassStationaryISMSwitcherProcessorTag, FInstancedActorsVisualizationProcessorTag>;
-
 	// some of the code below assumes EInstancedActorsBulkLOD::Detailed == 0, we need to verify that's the case. If not the code below needs updating.
 	static_assert((uint8)EInstancedActorsBulkLOD::Detailed == 0, "Code below relies on the assumptions. Needs to be updated if the assumption is broken");
 
@@ -341,7 +336,7 @@ void UInstancedActorsStationaryLODBatchProcessor::Execute(FMassEntityManager& En
 
 						if (ManagerSharedFragment.BulkLOD == EInstancedActorsBulkLOD::Detailed)
 						{
-							EntityManager.Defer().PushCommand<FAddRelevantTagsCommand>(InstanceData->Entities);
+							EntityManager.Defer().PushCommand<UE::InstancedActors::FEnableDetailedLODCommand>(InstanceData->Entities);
 						}
 						else
 						{
@@ -377,7 +372,8 @@ void UInstancedActorsStationaryLODBatchProcessor::Execute(FMassEntityManager& En
 									UMassStationaryISMSwitcherProcessor::ProcessContext(Context);
 								});
 
-							EntityManager.Defer().PushCommand<FRemoveRelevantTagsCommand>(InstanceData->Entities);
+							// Removes a bunch of tags from all mass entities that belong to an InstancedActorsData, so that we don't spend MassProcessor time on them
+							EntityManager.Defer().PushCommand<UE::InstancedActors::FEnableBatchLODCommand>(InstanceData->Entities);
 						}
 					}
 				}
diff --git a/Engine/Plugins/Runtime/InstancedActors/Source/InstancedActors/Public/InstancedActorsCommands.h b/Engine/Plugins/Runtime/InstancedActors/Source/InstancedActors/Public/InstancedActorsCommands.h
--- a/Engine/Plugins/Runtime/InstancedActors/Source/InstancedActors/Public/InstancedActorsCommands.h
+++ b/Engine/Plugins/Runtime/InstancedActors/Source/InstancedActors/Public/InstancedActorsCommands.h
@@ -6,6 +6,29 @@
 #include "MassCommands.h"
 #include "ProfilingDebugging/CpuProfilerTrace.h"
 
+namespace UE::InstancedActors
+{
+	/**
+	 * Returns the bitset indicating all the gate-tags of the processors we want to run on Detailed-LOD entities
+	 * (i.e. not the Batched-LOD ones). These tags are switched by UInstancedActorsStationaryLODBatchProcessor
+	 * Modifying the bitset is the way for project-specific code to influence what gets executed.
+	 */
+	INSTANCEDACTORS_API FMassTagBitSet& GetDetailedLODTags();
+
+	/** Adds GetDetailedLODTags() to an entity, effectively enabling DetailedLOD processing on it */
+	struct FEnableDetailedLODCommand : public FMassCommandChangeTags
+	{
+		using Super = FMassCommandChangeTags;
+		FEnableDetailedLODCommand();
+	};
+
+	/** Removes GetDetailedLODTags() from an entity, effectively enabling BatchLOD processing on it */
+	struct FEnableBatchLODCommand : public FMassCommandChangeTags
+	{
+		using Super = FMassCommandChangeTags;
+		FEnableBatchLODCommand();
+	};
+}
 
 /** 
  * Note: TManagerType is always expected to be AInstancedActorsManager, but is declared as 
