[InstancedActors] Refactored InstancedActorsData.GetBulkLOD to not fail if the instance has not been initialized yet.

diff --git a/Engine/Plugins/Runtime/InstancedActors/Source/InstancedActors/Public/InstancedActorsData.h b/Engine/Plugins/Runtime/InstancedActors/Source/InstancedActors/Public/InstancedActorsData.h
--- a/Engine/Plugins/Runtime/InstancedActors/Source/InstancedActors/Public/InstancedActorsData.h
+++ b/Engine/Plugins/Runtime/InstancedActors/Source/InstancedActors/Public/InstancedActorsData.h
@@ -75,7 +75,17 @@ public:
 	FMassEntityHandle GetEntity(FInstancedActorsInstanceIndex InstanceIndex) const;
 
 	void SetSharedInstancedActorDataStruct(FSharedStruct InSharedStruct);
-	EInstancedActorsBulkLOD GetBulkLOD() const { return SharedInstancedActorDataStruct.Get<FInstancedActorsDataSharedFragment>().BulkLOD; }
+	
+	/** @return current Bulk LOD of this InstancedActorData instance */
+	EInstancedActorsBulkLOD GetBulkLOD() const
+	{
+		if (const FInstancedActorsDataSharedFragment* DataFragment = SharedInstancedActorDataStruct.GetPtr<FInstancedActorsDataSharedFragment>())
+		{
+			return DataFragment->BulkLOD;
+		}
+		ensureMsgf(bHasEverInitialized == false, TEXT("We expect to hit this path only if the IAD instance has not been initialized yet."));
+		return EInstancedActorsBulkLOD::MAX;
+	}
 
 	// Identifying integer, unique within outer IAM, used for identifying matching persistence records
 	// Incrementally assigned in GetOrCreateActorInstanceData
