[InstancedActors] Fixed UInstancedActorsData::CreateEntityTemplate's creation of visualization trait. GetStationaryVisualizationTraitClass that's used as a souce of the class to spawn returns UMassStationaryDistanceVisualizationTrait and not UInstancedActorsVisualizationTrait (as it was before this change).
Also added explicit EntityTemplateID removal in UInstancedActorsData::ReleaseEntityTemplate since the one released through the EntityConfig is not the only once created by the given IAD.

diff --git a/Engine/Plugins/Runtime/InstancedActors/Source/InstancedActors/Private/InstancedActorsData.cpp b/Engine/Plugins/Runtime/InstancedActors/Source/InstancedActors/Private/InstancedActorsData.cpp
--- a/Engine/Plugins/Runtime/InstancedActors/Source/InstancedActors/Private/InstancedActorsData.cpp
+++ b/Engine/Plugins/Runtime/InstancedActors/Source/InstancedActors/Private/InstancedActorsData.cpp
@@ -281,7 +281,7 @@ void UInstancedActorsData::CreateEntityTemplate(const AActor& ExemplarActor)
 	UMassDistanceLODCollectorTrait* LODCollectorTrait = NewObject<UMassDistanceLODCollectorTrait>(this);
 	EntityConfig.AddTrait(*LODCollectorTrait);
 
-	UInstancedActorsVisualizationTrait* VisTrait = NewObject<UInstancedActorsVisualizationTrait>(this, GET_INSTANCEDACTORS_CONFIG_VALUE(GetStationaryVisualizationTraitClass()));
+	UMassStationaryDistanceVisualizationTrait* VisTrait = NewObject<UMassStationaryDistanceVisualizationTrait>(this, GET_INSTANCEDACTORS_CONFIG_VALUE(GetStationaryVisualizationTraitClass()));
 	if (UInstancedActorsVisualizationTrait* InstancedActorsVisTrait = Cast<UInstancedActorsVisualizationTrait>(VisTrait))
 	{
 		InstancedActorsVisTrait->InitializeFromInstanceData(*this);
@@ -289,8 +289,11 @@ void UInstancedActorsData::CreateEntityTemplate(const AActor& ExemplarActor)
 	EntityConfig.AddTrait(*VisTrait);
 
 	// Allow UInstancedActorsComponent's to extend entity config
-	ExemplarActor.ForEachComponent<UInstancedActorsComponent>(/*bIncludeFromChildActors*/ false, [this, &MassEntityManager](const UInstancedActorsComponent* InstancedActorComponent)
-		{ InstancedActorComponent->ModifyMassEntityConfig(MassEntityManager, this, EntityConfig); });
+	ExemplarActor.ForEachComponent<UInstancedActorsComponent>(/*bIncludeFromChildActors*/ false
+		, [this, &MassEntityManager](const UInstancedActorsComponent* InstancedActorComponent)
+		{
+			InstancedActorComponent->ModifyMassEntityConfig(MassEntityManager, this, EntityConfig);
+		});
 
 	const FMassEntityTemplate& BaseEntityTemplate = EntityConfig.GetOrCreateEntityTemplate(*World);
 
@@ -307,6 +310,7 @@ void UInstancedActorsData::CreateEntityTemplate(const AActor& ExemplarActor)
 
 	UMassSpawnerSubsystem* MassSpawnerSubsystem = UWorld::GetSubsystem<UMassSpawnerSubsystem>(GetWorld());
 	FMassEntityTemplateRegistry& TemplateRegistry = MassSpawnerSubsystem->GetMutableTemplateRegistryInstance();
+	// @todo we should have GUID per IAD and use it here.
 	const TSharedRef<FMassEntityTemplate>& FinalizedTemplate = TemplateRegistry.FindOrAddTemplate(FMassEntityTemplateIDFactory::Make(FGuid::NewGuid()), MoveTemp(ModifiedTemplate));
 	EntityTemplateID = FinalizedTemplate->GetTemplateID();
 }
@@ -324,6 +328,10 @@ void UInstancedActorsData::ReleaseEntityTemplate()
 	check(World);
 
 	EntityConfig.DestroyEntityTemplate(*World);
+
+	UMassSpawnerSubsystem* MassSpawnerSubsystem = UWorld::GetSubsystem<UMassSpawnerSubsystem>(World);
+	FMassEntityTemplateRegistry& TemplateRegistry = MassSpawnerSubsystem->GetMutableTemplateRegistryInstance();
+	TemplateRegistry.DestroyTemplate(EntityTemplateID);
 }
 
 void UInstancedActorsData::SpawnEntities()
