[InstancedActors] FInstancedActorsSettings can now override and set FGameplayTagContainer which can be used used by UInstancedActorsModifierBase::InstanceTagsQuery
- Changed UInstancedActorsData::Tags to AdditionalTags. UInstancedActorsSubsystem::InstanceActor() and FPCGSpawnActorForPath will now pass these along as additional tags
- UInstancedActorsData::Initialize() will combine these AdditionalTags with the new FInstancedActorsSettings::GameplayTags. This new CombinedTags will be used for queries

diff --git a/Engine/Plugins/Runtime/InstancedActors/Config/BaseInstancedActors.ini b/Engine/Plugins/Runtime/InstancedActors/Config/BaseInstancedActors.ini
--- a/Engine/Plugins/Runtime/InstancedActors/Config/BaseInstancedActors.ini
+++ b/Engine/Plugins/Runtime/InstancedActors/Config/BaseInstancedActors.ini
@@ -1 +1,2 @@
 [CoreRedirects]
++PropertyRedirects=(OldName="/Script/InstancedActors.InstancedActorsData.Tags",NewName="/Script/InstancedActors.InstancedActorsData.AdditionalTags")
diff --git a/Engine/Plugins/Runtime/InstancedActors/Source/InstancedActors/Private/InstancedActorsData.cpp b/Engine/Plugins/Runtime/InstancedActors/Source/InstancedActors/Private/InstancedActorsData.cpp
--- a/Engine/Plugins/Runtime/InstancedActors/Source/InstancedActors/Private/InstancedActorsData.cpp
+++ b/Engine/Plugins/Runtime/InstancedActors/Source/InstancedActors/Private/InstancedActorsData.cpp
@@ -240,6 +240,18 @@ void UInstancedActorsData::Initialize()
 			}
 		}
 	}
+	
+	CombinedTags = AdditionalTags.GetTags();
+
+	// Allow settings to combine its tags with AdditionalTags
+	// @todo: resolve ambiguity with how the "overrides" should work with GameplayTags. The override flag is used both to solve inheritance-based
+	// priority overriding of values, and to denote that a logical override will apply to IAD behavior. However, the end-result use case of
+	// GameplayTags for IADs is to merge the settings-based tags with those that may have been baked into the IAD (i.e., AdditionalTags).
+	// We may need to add a "merge" functionality to settings that would work for GameplayTags
+	if (Settings && Settings->bOverride_GameplayTags)
+	{
+		CombinedTags.AppendTags(Settings->GameplayTags);
+	}
 
 	// Get or create exemplar actor to derive entities from
 	check(!ExemplarActorData.IsValid());
diff --git a/Engine/Plugins/Runtime/InstancedActors/Source/InstancedActors/Private/InstancedActorsManager.cpp b/Engine/Plugins/Runtime/InstancedActors/Source/InstancedActors/Private/InstancedActorsManager.cpp
--- a/Engine/Plugins/Runtime/InstancedActors/Source/InstancedActors/Private/InstancedActorsManager.cpp
+++ b/Engine/Plugins/Runtime/InstancedActors/Source/InstancedActors/Private/InstancedActorsManager.cpp
@@ -670,7 +670,7 @@ void AInstancedActorsManager::SetGridGuid(const FGuid& InGuid)
 	ManagerGridGuid = InGuid;
 }
 
-UInstancedActorsData* AInstancedActorsManager::CreateNextInstanceActorData(TSubclassOf<AActor> ActorClass, const FInstancedActorsTagSet& InstanceTags)
+UInstancedActorsData* AInstancedActorsManager::CreateNextInstanceActorData(TSubclassOf<AActor> ActorClass, const FInstancedActorsTagSet& AdditionalInstanceTags)
 {
 	check(InstancedActorsDataClass);
 
@@ -679,7 +679,7 @@ UInstancedActorsData* AInstancedActorsManager::CreateNextInstanceActorData(TSubc
 	// @todo it's conceivable the NextInstanceDataID will overflow. We need to use some handle system in place instead. 
 	NewInstanceData->ID = NextInstanceDataID++;
 	NewInstanceData->ActorClass = ActorClass;
-	NewInstanceData->Tags = InstanceTags;
+	NewInstanceData->AdditionalTags = AdditionalInstanceTags;
 	check(Algo::NoneOf(PerActorClassInstanceData, [NewInstanceData](UInstancedActorsData* InstanceData)
 		{
 			return InstanceData->ID == NewInstanceData->ID;
@@ -688,13 +688,13 @@ UInstancedActorsData* AInstancedActorsManager::CreateNextInstanceActorData(TSubc
 	return NewInstanceData;
 }
 
-UInstancedActorsData& AInstancedActorsManager::GetOrCreateActorInstanceData(TSubclassOf<AActor> ActorClass, const FInstancedActorsTagSet& InstanceTags)
+UInstancedActorsData& AInstancedActorsManager::GetOrCreateActorInstanceData(TSubclassOf<AActor> ActorClass, const FInstancedActorsTagSet& AdditionalInstanceTags)
 {
 	checkf(HasActorBegunPlay() == false, TEXT("AInstancedActorsManager doesn't yet support runtime addition of instances"));
 
-	TObjectPtr<UInstancedActorsData>* InstanceData = PerActorClassInstanceData.FindByPredicate([ActorClass, InstanceTags](TObjectPtr<UInstancedActorsData> InstanceData)
+	TObjectPtr<UInstancedActorsData>* InstanceData = PerActorClassInstanceData.FindByPredicate([ActorClass, AdditionalInstanceTags](TObjectPtr<UInstancedActorsData> InstanceData)
 		{ 
-			return InstanceData->ActorClass == ActorClass && InstanceData->Tags == InstanceTags; 
+			return InstanceData->ActorClass == ActorClass && InstanceData->AdditionalTags == AdditionalInstanceTags; 
 		});
 
 	if (InstanceData != nullptr)
@@ -703,7 +703,7 @@ UInstancedActorsData& AInstancedActorsManager::GetOrCreateActorInstanceData(TSub
 		return **InstanceData;
 	}
 
-	UInstancedActorsData* NewInstanceData = CreateNextInstanceActorData(ActorClass, InstanceTags);
+	UInstancedActorsData* NewInstanceData = CreateNextInstanceActorData(ActorClass, AdditionalInstanceTags);
 	PerActorClassInstanceData.Add(NewInstanceData);
 
 	// Get or create exemplar actor to derive entities from
@@ -760,9 +760,9 @@ void AInstancedActorsManager::PreRegisterAllComponents()
 	}
 }
 
-FInstancedActorsInstanceHandle AInstancedActorsManager::AddActorInstance(TSubclassOf<AActor> ActorClass, FTransform InstanceTransform, bool bWorldSpace, const FInstancedActorsTagSet& InstanceTags)
+FInstancedActorsInstanceHandle AInstancedActorsManager::AddActorInstance(TSubclassOf<AActor> ActorClass, FTransform InstanceTransform, bool bWorldSpace, const FInstancedActorsTagSet& AdditionalInstanceTags)
 {
-	UInstancedActorsData& InstanceData = GetOrCreateActorInstanceData(ActorClass, InstanceTags);
+	UInstancedActorsData& InstanceData = GetOrCreateActorInstanceData(ActorClass, AdditionalInstanceTags);
 	return InstanceData.AddInstance(InstanceTransform, bWorldSpace);
 }
 
diff --git a/Engine/Plugins/Runtime/InstancedActors/Source/InstancedActors/Private/InstancedActorsModifiers.cpp b/Engine/Plugins/Runtime/InstancedActors/Source/InstancedActors/Private/InstancedActorsModifiers.cpp
--- a/Engine/Plugins/Runtime/InstancedActors/Source/InstancedActors/Private/InstancedActorsModifiers.cpp
+++ b/Engine/Plugins/Runtime/InstancedActors/Source/InstancedActors/Private/InstancedActorsModifiers.cpp
@@ -23,7 +23,7 @@ void UInstancedActorsModifierBase::ModifyAllInstances(AInstancedActorsManager& M
 	IterationContext,
 	/*Predicate*/InstanceTagsQuery.IsEmpty() ? TOptional<AInstancedActorsManager::FInstancedActorDataPredicateFunc>() : TOptional<AInstancedActorsManager::FInstancedActorDataPredicateFunc>([this](const UInstancedActorsData& InstancedActorData)
 	{
-		return InstancedActorData.Tags.GetTags().MatchesQuery(InstanceTagsQuery);
+		return InstancedActorData.GetCombinedTags().MatchesQuery(InstanceTagsQuery);
 	}));
 }
 
@@ -53,7 +53,7 @@ void URemoveInstancedActorsModifier::ModifyAllInstances(AInstancedActorsManager&
 		{
 			check(IsValid(InstanceData));
 
-			if (InstanceData->Tags.GetTags().MatchesQuery(InstanceTagsQuery))
+			if (InstanceData->GetCombinedTags().MatchesQuery(InstanceTagsQuery))
 			{
 				IterationContext.RemoveAllInstancesDeferred(*InstanceData);
 			}
diff --git a/Engine/Plugins/Runtime/InstancedActors/Source/InstancedActors/Private/InstancedActorsSettingsTypes.cpp b/Engine/Plugins/Runtime/InstancedActors/Source/InstancedActors/Private/InstancedActorsSettingsTypes.cpp
--- a/Engine/Plugins/Runtime/InstancedActors/Source/InstancedActors/Private/InstancedActorsSettingsTypes.cpp
+++ b/Engine/Plugins/Runtime/InstancedActors/Source/InstancedActors/Private/InstancedActorsSettingsTypes.cpp
@@ -100,7 +100,16 @@ static inline void DisplaySettings(FStringBuilderBase & SettingsString, const TA
 			SettingsString << FString::SanitizeFloat((float)SettingsValue[ui]) << TEXT(" ");
 		}
 	}
+}
 
+static inline void DisplaySettings(FStringBuilderBase & SettingsString, const FGameplayTagContainer& GameplayTags, bool bOverridesOnly, bool bOverride_Settings, const TCHAR* SettingsName)
+{
+	if (!bOverridesOnly || bOverride_Settings)
+	{
+		SettingsString << SettingsName << TEXT(": ");
+
+		SettingsString << GameplayTags.ToStringSimple();
+	}
 }
 
 void FInstancedActorsSettings::OverrideIfDefault(FConstStructView InOverrideSettings, const FName& OverrideSettingsName)
@@ -125,6 +134,7 @@ void FInstancedActorsSettings::OverrideIfDefault(FConstStructView InOverrideSett
 	IASETTINGS_OVERRIDE_IF_DEFAULT(bCanBeDamaged);
 	IASETTINGS_OVERRIDE_IF_DEFAULT(bIgnoreModifierVolumes);
 	IASETTINGS_OVERRIDE_IF_DEFAULT(bControlPhysicsState);
+	IASETTINGS_OVERRIDE_IF_DEFAULT(GameplayTags);
 	
 	AppliedSettingsOverrides.Add(OverrideSettingsName);
 }
@@ -150,6 +160,7 @@ FString FInstancedActorsSettings::DebugToString(bool bOverridesOnly) const
 	IASETTINGS_SETTING_TO_STRING(bCanBeDamaged);
 	IASETTINGS_SETTING_TO_STRING(bIgnoreModifierVolumes);
 	IASETTINGS_SETTING_TO_STRING(bControlPhysicsState);
+	DisplaySettings(SettingsString, GameplayTags, bOverridesOnly, bOverride_GameplayTags, TEXT("GameplayTags"));
 
 	DisplaySettings(SettingsString, LODDistanceScales, bOverridesOnly, bOverride_LODDistanceScales, TEXT("LODDistanceScales"));
 
diff --git a/Engine/Plugins/Runtime/InstancedActors/Source/InstancedActors/Private/InstancedActorsSubsystem.cpp b/Engine/Plugins/Runtime/InstancedActors/Source/InstancedActors/Private/InstancedActorsSubsystem.cpp
--- a/Engine/Plugins/Runtime/InstancedActors/Source/InstancedActors/Private/InstancedActorsSubsystem.cpp
+++ b/Engine/Plugins/Runtime/InstancedActors/Source/InstancedActors/Private/InstancedActorsSubsystem.cpp
@@ -388,12 +388,12 @@ void UInstancedActorsSubsystem::RemoveModifierVolume(const FInstancedActorsModif
 }
 
 #if WITH_EDITOR
-FInstancedActorsInstanceHandle UInstancedActorsSubsystem::InstanceActor(TSubclassOf<AActor> ActorClass, FTransform InstanceTransform, ULevel* Level, const FGameplayTagContainer& InstanceTags)
+FInstancedActorsInstanceHandle UInstancedActorsSubsystem::InstanceActor(TSubclassOf<AActor> ActorClass, FTransform InstanceTransform, ULevel* Level, const FGameplayTagContainer& AdditionalInstanceTags)
 {
-	return InstanceActor(ActorClass, InstanceTransform, Level, InstanceTags, InstancedActorsManagerClass);
+	return InstanceActor(ActorClass, InstanceTransform, Level, AdditionalInstanceTags, InstancedActorsManagerClass);
 }
 
-FInstancedActorsInstanceHandle UInstancedActorsSubsystem::InstanceActor(TSubclassOf<AActor> ActorClass, FTransform InstanceTransform, ULevel* Level, const FGameplayTagContainer& InstanceTags, TSubclassOf<AInstancedActorsManager> ManagerClass)
+FInstancedActorsInstanceHandle UInstancedActorsSubsystem::InstanceActor(TSubclassOf<AActor> ActorClass, FTransform InstanceTransform, ULevel* Level, const FGameplayTagContainer& AdditionalInstanceTags, TSubclassOf<AInstancedActorsManager> ManagerClass)
 {
 	if (!ensureMsgf(Level, TEXT("Expecting a valid Level. Received nullptr.")))
 	{
@@ -434,7 +434,7 @@ FInstancedActorsInstanceHandle UInstancedActorsSubsystem::InstanceActor(TSubclas
 					MessageLog.AddMessage(
 						FTokenizedMessage::Create(EMessageSeverity::Warning)
 							->AddToken(FUObjectToken::Create(ActorClass))
-							->AddToken(FTextToken::Create(NSLOCTEXT("InstancedActorss", "CantInstanceClassWarning", "doesn't have a matching class or super class entry in the ActorClassSettingsRegistry.")))
+							->AddToken(FTextToken::Create(NSLOCTEXT("InstancedActors", "CantInstanceClassWarning", "doesn't have a matching class or super class entry in the ActorClassSettingsRegistry.")))
 					);
 
 					MessageLog.Open(EMessageSeverity::Warning);
@@ -449,7 +449,7 @@ FInstancedActorsInstanceHandle UInstancedActorsSubsystem::InstanceActor(TSubclas
 					MessageLog.AddMessage(
 						FTokenizedMessage::Create(EMessageSeverity::Error)
 						->AddToken(FUObjectToken::Create(ActorClass))
-						->AddToken(FTextToken::Create(NSLOCTEXT("InstancedActorss", "CantInstanceClassError", "doesn't have a matching class or super class entry in the ActorClassSettingsRegistry, skipping instance of 'unknown' type.")))
+						->AddToken(FTextToken::Create(NSLOCTEXT("InstancedActors", "CantInstanceClassError", "doesn't have a matching class or super class entry in the ActorClassSettingsRegistry, skipping instance of 'unknown' type.")))
 					);
 
 					if (InstancedActorsCVars::EditorEnforceActorClassSettingsPresence >= 3)
@@ -514,7 +514,7 @@ FInstancedActorsInstanceHandle UInstancedActorsSubsystem::InstanceActor(TSubclas
 			AInstancedActorsManager* Manager = CastChecked<AInstancedActorsManager>(PartitionActor);
 	
 			// Add instance to manager
-			return Manager->AddActorInstance(ActorClass, InstanceTransform, /*bWorldSpace*/true, InstanceTags);
+			return Manager->AddActorInstance(ActorClass, InstanceTransform, /*bWorldSpace*/true, AdditionalInstanceTags);
 		}
 	}
 
diff --git a/Engine/Plugins/Runtime/InstancedActors/Source/InstancedActors/Public/InstancedActorsData.h b/Engine/Plugins/Runtime/InstancedActors/Source/InstancedActors/Public/InstancedActorsData.h
--- a/Engine/Plugins/Runtime/InstancedActors/Source/InstancedActors/Public/InstancedActorsData.h
+++ b/Engine/Plugins/Runtime/InstancedActors/Source/InstancedActors/Public/InstancedActorsData.h
@@ -96,9 +96,10 @@ public:
 	UPROPERTY(VisibleAnywhere, Category=InstancedActors)
 	TSubclassOf<AActor> ActorClass;
 
-	// Delimiting tag set to group sets of instances
-	UPROPERTY(VisibleAnywhere, Category=InstancedActors)
-	FInstancedActorsTagSet Tags;
+	const FInstancedActorsTagSet& GetAdditionalTags() const
+	{
+		return AdditionalTags;
+	}
 
 	template<typename T>
 	const std::remove_const_t<T>* GetSettingsPtr() const
@@ -118,6 +119,10 @@ public:
 	}
 
 protected:
+	// Delimiting tag set to group sets of instances
+	UPROPERTY(VisibleAnywhere, Category=InstancedActors)
+	FInstancedActorsTagSet AdditionalTags;
+	
 	// Compiled settings for these instances, compiled and cached in BeginPlay.
 	// Note: Can also be refreshed on demand in PIE, using IA.RefreshSettings CVar
 	FConstSharedStruct SharedSettings;
@@ -395,6 +400,8 @@ public:
 
 	int32 GetEntityIndexFromCollisionIndex(const UInstancedStaticMeshComponent& ISMComponent, const int32 CollisionIndex) const;
 
+	const FGameplayTagContainer& GetCombinedTags() const { return CombinedTags; }
+	
 protected:
 	FMassEntityHandle GetEntityHandleForIndex(const int32 Index) const
 	{
@@ -490,5 +497,9 @@ private:
 	UPROPERTY(Transient)
 	FMassEntityConfig EntityConfig;
 
+	// Result of combining AdditionalTags with class-based FInstancedActorsSettings::GameplayTags
+	UPROPERTY(Transient)
+	FGameplayTagContainer CombinedTags;
+
 	TSharedPtr<UE::InstancedActors::FExemplarActorData> ExemplarActorData;
 };
diff --git a/Engine/Plugins/Runtime/InstancedActors/Source/InstancedActors/Public/InstancedActorsManager.h b/Engine/Plugins/Runtime/InstancedActors/Source/InstancedActors/Public/InstancedActorsManager.h
--- a/Engine/Plugins/Runtime/InstancedActors/Source/InstancedActors/Public/InstancedActorsManager.h
+++ b/Engine/Plugins/Runtime/InstancedActors/Source/InstancedActors/Public/InstancedActorsManager.h
@@ -83,7 +83,7 @@ public:
 
 #if WITH_EDITOR
 	/** Adds an instance of ActorClass at InstanceTransform location to instance data */
-	FInstancedActorsInstanceHandle AddActorInstance(TSubclassOf<AActor> ActorClass, FTransform InstanceTransform, bool bWorldSpace = true, const FInstancedActorsTagSet& InstanceTags = FInstancedActorsTagSet());
+	FInstancedActorsInstanceHandle AddActorInstance(TSubclassOf<AActor> ActorClass, FTransform InstanceTransform, bool bWorldSpace = true, const FInstancedActorsTagSet& AdditionalInstanceTags = FInstancedActorsTagSet());
 
 	/**
 	 * Removes all instance data for InstanceHandle.
@@ -291,10 +291,10 @@ protected:
 	FBox CalculateLocalInstanceBounds() const;
 
 #if WITH_EDITOR
-	/** Helper function to create and initialize per-actor-class UInstancedActorsData's, optionally further partitioned by InstanceTags */
-	UInstancedActorsData& GetOrCreateActorInstanceData(TSubclassOf<AActor> ActorClass, const FInstancedActorsTagSet& InstanceTags);
+	/** Helper function to create and initialize per-actor-class UInstancedActorsData's, optionally further partitioned by AdditionalInstanceTags */
+	UInstancedActorsData& GetOrCreateActorInstanceData(TSubclassOf<AActor> ActorClass, const FInstancedActorsTagSet& AdditionalInstanceTags);
 
-	virtual UInstancedActorsData* CreateNextInstanceActorData(TSubclassOf<AActor> ActorClass, const FInstancedActorsTagSet& InstanceTags);
+	virtual UInstancedActorsData* CreateNextInstanceActorData(TSubclassOf<AActor> ActorClass, const FInstancedActorsTagSet& AdditionalInstanceTags);
 
 	/** Used to set the right properties on the editor ISMCs so we can do per-instance selection. */
 	virtual void PreRegisterAllComponents() override;
diff --git a/Engine/Plugins/Runtime/InstancedActors/Source/InstancedActors/Public/InstancedActorsModifiers.h b/Engine/Plugins/Runtime/InstancedActors/Source/InstancedActors/Public/InstancedActorsModifiers.h
--- a/Engine/Plugins/Runtime/InstancedActors/Source/InstancedActors/Public/InstancedActorsModifiers.h
+++ b/Engine/Plugins/Runtime/InstancedActors/Source/InstancedActors/Public/InstancedActorsModifiers.h
@@ -100,7 +100,7 @@ public:
 
 			if (!InstanceTagsQuery.IsEmpty())
 			{
-				return InstancedActorData.Tags.GetTags().MatchesQuery(InstanceTagsQuery);
+				return InstancedActorData.GetCombinedTags().MatchesQuery(InstanceTagsQuery);
 			}
 			
 			return true;
diff --git a/Engine/Plugins/Runtime/InstancedActors/Source/InstancedActors/Public/InstancedActorsSettingsTypes.h b/Engine/Plugins/Runtime/InstancedActors/Source/InstancedActors/Public/InstancedActorsSettingsTypes.h
--- a/Engine/Plugins/Runtime/InstancedActors/Source/InstancedActors/Public/InstancedActorsSettingsTypes.h
+++ b/Engine/Plugins/Runtime/InstancedActors/Source/InstancedActors/Public/InstancedActorsSettingsTypes.h
@@ -6,6 +6,7 @@
 #include "GameFramework/Actor.h"
 #include "StructUtils/StructView.h"
 #include "StructUtils/InstancedStruct.h"
+#include "GameplayTagContainer.h"
 #include "InstancedActorsSettingsTypes.generated.h"
 
 
@@ -125,6 +126,9 @@ struct INSTANCEDACTORS_API FInstancedActorsSettings : public FTableRowBase
 	UPROPERTY(EditAnywhere, BlueprintReadWrite, Meta=(InlineEditConditionToggle), Category=InstancedActors)
 	uint8 bOverride_bControlPhysicsState : 1 = false;
 
+	UPROPERTY(EditAnywhere, BlueprintReadWrite, Meta=(InlineEditConditionToggle), Category=InstancedActors)
+	uint8 bOverride_GameplayTags : 1 = false;
+
 	// Settings 
 
 	/** Optional shadow casting override applied to instance ISMC's if set (shadow casting settings from ActorClass will be used for ISMC's if unset) */
@@ -212,6 +216,10 @@ struct INSTANCEDACTORS_API FInstancedActorsSettings : public FTableRowBase
 	/** Completely disable modifier volumes for this type of instance. */
 	UPROPERTY(EditAnywhere, BlueprintReadWrite, Meta = (EditCondition = "bOverride_bIgnoreModifierVolumes"), Category=InstancedActors)
 	bool bIgnoreModifierVolumes = false;
+	
+	/** Gameplay Tags that can be used in Modifier queries */
+	UPROPERTY(EditAnywhere, BlueprintReadWrite, Meta = (EditCondition = "bOverride_GameplayTags"), Category=InstancedActors)
+	FGameplayTagContainer GameplayTags;
 
 	// Note: Don't forget to implement overriding for new settings in ApplyOverrides and Stringification in ToString
 	UPROPERTY(VisibleAnywhere, Category = InstancedActors)
diff --git a/Engine/Plugins/Runtime/InstancedActors/Source/InstancedActors/Public/InstancedActorsSubsystem.h b/Engine/Plugins/Runtime/InstancedActors/Source/InstancedActors/Public/InstancedActorsSubsystem.h
--- a/Engine/Plugins/Runtime/InstancedActors/Source/InstancedActors/Public/InstancedActorsSubsystem.h
+++ b/Engine/Plugins/Runtime/InstancedActors/Source/InstancedActors/Public/InstancedActorsSubsystem.h
@@ -47,18 +47,20 @@ public:
 #if WITH_EDITOR
 	/** 
 	* Adds an instance of ActorClass at InstanceTransform location by spawning or reusing a AInstancedActorsManager at InstanceTransform's grid cell location.
+	* @param AdditionalInstanceTags Gameplay tags that will get appended to class-based FInstancedActorsSettings::GameplayTags
 	* @see UInstancedActorsProjectSettings::GridSize
 	*/	
 	INSTANCEDACTORS_API FInstancedActorsInstanceHandle InstanceActor(TSubclassOf<AActor> ActorClass, FTransform InstanceTransform, ULevel* Level
-		, const FGameplayTagContainer& InstanceTags = FGameplayTagContainer());
+		, const FGameplayTagContainer& AdditionalInstanceTags = FGameplayTagContainer());
 
 	/** 
 	* Adds an instance of ActorClass at InstanceTransform location by spawning or reusing a AInstancedActorsManager at InstanceTransform's grid cell location.
+	* * @param AdditionalInstanceTags Gameplay tags that will get appended to class-based FInstancedActorsSettings::GameplayTags
 	* @see UInstancedActorsProjectSettings::GridSize
 	*/	
 	UFUNCTION(BlueprintCallable, Category = InstancedActors)
 	INSTANCEDACTORS_API FInstancedActorsInstanceHandle InstanceActor(TSubclassOf<AActor> ActorClass, FTransform InstanceTransform, ULevel* Level
-		, const FGameplayTagContainer& InstanceTags, TSubclassOf<AInstancedActorsManager> ManagerClass);
+		, const FGameplayTagContainer& AdditionalInstanceTags, TSubclassOf<AInstancedActorsManager> ManagerClass);
 
 	/**
 	 * Removes all instance data for InstanceHandle.	
