[InstancedActors] relaxed conditions for invalid representations as checked by UInstancedActorsStationaryLODBatchProcessor::Execute.

We used to tolerate invalid representations only on clients. We now will also let them be when there's no outer world.

The main point of the change is to catch a recurring failure in existing projects.

diff --git a/Engine/Plugins/Runtime/InstancedActors/Source/InstancedActors/Private/InstancedActorsStationaryLODBatchProcessor.cpp b/Engine/Plugins/Runtime/InstancedActors/Source/InstancedActors/Private/InstancedActorsStationaryLODBatchProcessor.cpp
--- a/Engine/Plugins/Runtime/InstancedActors/Source/InstancedActors/Private/InstancedActorsStationaryLODBatchProcessor.cpp
+++ b/Engine/Plugins/Runtime/InstancedActors/Source/InstancedActors/Private/InstancedActorsStationaryLODBatchProcessor.cpp
@@ -408,7 +408,7 @@ void UInstancedActorsStationaryLODBatchProcessor::Execute(FMassEntityManager& En
 
 		if (DirtyRepresentationInstances.Num())
 		{
-			const bool bIsClient = (InstancedActorSubsystem->GetWorld() && InstancedActorSubsystem->GetWorld()->GetNetMode() == NM_Client);
+			const bool bInvalidInstancedAllowed = (InstancedActorSubsystem->GetWorld() == nullptr) || (InstancedActorSubsystem->GetWorld()->GetNetMode() == NM_Client);
 
 			// Collect mass entities from instance handles into entity collections for processing
 			// UE::Mass::Utils::CreateEntityCollections but from TArray<FInstancedActorsInstanceHandle>, retrieving instance entities as we go
@@ -416,8 +416,9 @@ void UInstancedActorsStationaryLODBatchProcessor::Execute(FMassEntityManager& En
 			for (const FInstancedActorsInstanceHandle& DirtyRepresentationInstance : DirtyRepresentationInstances)
 			{
 				// Note that it's possible for DirtyRepresentationInstances to contain indices to entities that just have 
-				// been destroyed (however, we only expect it to happen only on clients). 
-				ensureMsgf(DirtyRepresentationInstance.IsValid() || bIsClient, TEXT("We only expect invalid instance handles on Client."));
+				// been destroyed (however, we only expect it to happen only on clients, of if we're somewhere in the process of destroying the world). 
+				ensureMsgf(DirtyRepresentationInstance.IsValid() || bInvalidInstancedAllowed
+					, TEXT("We only expect invalid instance handles on Client or when the InstancedActorSubsystem no longer has a valid outer UWorld."));
 
 				if (DirtyRepresentationInstance.IsValid() 
 					// only the entities that are not "Detailed" require update, so we're filtering out accordingly
