[Mass] Issue 624880: added a bunch of World and other headers in hopes of fixing the non-unit compilation issue.

#rnx

diff --git a/Engine/Plugins/AI/MassAI/Source/MassAIBehavior/Private/MassComponentHitSubsystem.cpp b/Engine/Plugins/AI/MassAI/Source/MassAIBehavior/Private/MassComponentHitSubsystem.cpp
--- a/Engine/Plugins/AI/MassAI/Source/MassAIBehavior/Private/MassComponentHitSubsystem.cpp
+++ b/Engine/Plugins/AI/MassAI/Source/MassAIBehavior/Private/MassComponentHitSubsystem.cpp
@@ -8,6 +8,8 @@
 #include "MassSignalSubsystem.h"
 #include "MassSimulationSubsystem.h"
 #include "Components/CapsuleComponent.h"
+#include "Subsystems/SubsystemCollection.h"
+#include "Engine/World.h"
 
 namespace UE::MassComponentHit
 {
diff --git a/Engine/Plugins/Runtime/InstancedActors/Source/InstancedActors/Private/GameplayDebuggerCategory_InstancedActors.cpp b/Engine/Plugins/Runtime/InstancedActors/Source/InstancedActors/Private/GameplayDebuggerCategory_InstancedActors.cpp
--- a/Engine/Plugins/Runtime/InstancedActors/Source/InstancedActors/Private/GameplayDebuggerCategory_InstancedActors.cpp
+++ b/Engine/Plugins/Runtime/InstancedActors/Source/InstancedActors/Private/GameplayDebuggerCategory_InstancedActors.cpp
@@ -18,6 +18,8 @@
 #include "MassDebuggerSubsystem.h"
 #include "MassEntitySubsystem.h"
 #include "MassEntityView.h"
+#include "Engine/World.h"
+#include "MassEntitySubsystem.h"
 
 
 namespace UE::InstancedActors::Debug
diff --git a/Engine/Plugins/Runtime/InstancedActors/Source/InstancedActors/Private/InstancedActorsModifierVolumeComponent.cpp b/Engine/Plugins/Runtime/InstancedActors/Source/InstancedActors/Private/InstancedActorsModifierVolumeComponent.cpp
--- a/Engine/Plugins/Runtime/InstancedActors/Source/InstancedActors/Private/InstancedActorsModifierVolumeComponent.cpp
+++ b/Engine/Plugins/Runtime/InstancedActors/Source/InstancedActors/Private/InstancedActorsModifierVolumeComponent.cpp
@@ -7,7 +7,7 @@
 #include "InstancedActorsModifiers.h"
 #include "InstancedActorsSubsystem.h"
 #include "InstancedActorsTypes.h"
-
+#include "Engine/World.h"
 #include "Algo/NoneOf.h"
 #include "Net/UnrealNetwork.h"
 #include "PrimitiveSceneProxy.h"
diff --git a/Engine/Plugins/Runtime/InstancedActors/Source/InstancedActors/Public/GameplayDebuggerCategory_InstancedActors.h b/Engine/Plugins/Runtime/InstancedActors/Source/InstancedActors/Public/GameplayDebuggerCategory_InstancedActors.h
--- a/Engine/Plugins/Runtime/InstancedActors/Source/InstancedActors/Public/GameplayDebuggerCategory_InstancedActors.h
+++ b/Engine/Plugins/Runtime/InstancedActors/Source/InstancedActors/Public/GameplayDebuggerCategory_InstancedActors.h
@@ -9,7 +9,7 @@
 #if WITH_GAMEPLAY_DEBUGGER && WITH_INSTANCEDACTORS_DEBUG && WITH_MASSENTITY_DEBUG
 
 #include "GameplayDebuggerCategory.h"
-#include "MassEntityTypes.h"
+#include "HAL/IConsoleManager.h"
 
 struct FMassEntityManager;
 
diff --git a/Engine/Source/Developer/MassEntityTestSuite/Private/MassProcessingPhasesTest.cpp b/Engine/Source/Developer/MassEntityTestSuite/Private/MassProcessingPhasesTest.cpp
--- a/Engine/Source/Developer/MassEntityTestSuite/Private/MassProcessingPhasesTest.cpp
+++ b/Engine/Source/Developer/MassEntityTestSuite/Private/MassProcessingPhasesTest.cpp
@@ -5,6 +5,7 @@
 #include "MassProcessingPhaseManager.h"
 #include "MassEntityTestTypes.h"
 #include "MassProcessingPhaseManager.h"
+#include "Engine/World.h"
 
 #define LOCTEXT_NAMESPACE "MassTest"
 
diff --git a/Engine/Source/Runtime/MassEntity/Private/MassEntitySubsystem.cpp b/Engine/Source/Runtime/MassEntity/Private/MassEntitySubsystem.cpp
--- a/Engine/Source/Runtime/MassEntity/Private/MassEntitySubsystem.cpp
+++ b/Engine/Source/Runtime/MassEntity/Private/MassEntitySubsystem.cpp
@@ -2,6 +2,7 @@
 
 #include "MassEntitySubsystem.h"
 #include "HAL/IConsoleManager.h"
+#include "Engine/World.h"
 
 #include UE_INLINE_GENERATED_CPP_BY_NAME(MassEntitySubsystem)
 
diff --git a/Engine/Source/Runtime/MassEntity/Public/MassProcessingTypes.h b/Engine/Source/Runtime/MassEntity/Public/MassProcessingTypes.h
--- a/Engine/Source/Runtime/MassEntity/Public/MassProcessingTypes.h
+++ b/Engine/Source/Runtime/MassEntity/Public/MassProcessingTypes.h
@@ -4,6 +4,7 @@
 
 #include "StructUtils/StructUtilsTypes.h"
 #include "StructUtils/InstancedStruct.h"
+#include "Templates/SubclassOf.h"
 #include "MassProcessingTypes.generated.h"
 
 #ifndef MASS_DO_PARALLEL
