[InstancedActors] Fixed ensures poping up when the server replicates over actor information before the client-side IAD even spawned its entities.
Solving the issue by caching the requests and reapplying after spawning

diff --git a/Engine/Plugins/Runtime/InstancedActors/Source/InstancedActors/Private/InstancedActorsData.cpp b/Engine/Plugins/Runtime/InstancedActors/Source/InstancedActors/Private/InstancedActorsData.cpp
--- a/Engine/Plugins/Runtime/InstancedActors/Source/InstancedActors/Private/InstancedActorsData.cpp
+++ b/Engine/Plugins/Runtime/InstancedActors/Source/InstancedActors/Private/InstancedActorsData.cpp
@@ -279,6 +279,8 @@ void UInstancedActorsData::Deinitialize()
 
 		ExemplarActorData.Reset();
 	}
+
+	CachedSetReplicatedActorRequests.Empty();
 }
 
 void UInstancedActorsData::CreateEntityTemplate(const AActor& ExemplarActor)
@@ -404,6 +406,17 @@ void UInstancedActorsData::SpawnEntities()
 			InstancedActorSubsystem->UpdateAndResetTickTime(SharedInstancedActorDataStruct);
 		}
 	}
+
+	// only once we spawned entities we can apply the "SetReplicatedActor" that came in before this InstancedActorData was ready to receive them.
+	for (const FSetReplicatedActorRequests& Request : CachedSetReplicatedActorRequests)
+	{
+		if (AActor* Actor = Request.Actor.Get())
+		{
+			SetReplicatedActor(Request.InstanceIndex, *Actor);
+		}
+	}
+	// we can afford to release the memory, since this array is only ever expected to be used at IAD's boot up
+	CachedSetReplicatedActorRequests.Empty();
 }
 
 void UInstancedActorsData::DespawnEntities()
@@ -1245,6 +1258,13 @@ void UInstancedActorsData::SetReplicatedActor(FInstancedActorsInstanceIndex Inst
 {
 	checkf(!ReplicatedActor.HasAuthority(), TEXT("SetReplicatedActor must only be called on clients where we explictly connect actor to entities via UInstancedActorsComponent::OnRep_InstanceHandle"));
 
+	if (HasSpawnedEntities() == false)
+	{	
+		// we need to cache the request to apply it once we have spawned the entities
+		CachedSetReplicatedActorRequests.Emplace(Instance, &ReplicatedActor);
+		return;
+	}
+
 	if (!ensure(Entities.IsValidIndex(Instance.GetIndex())))
 	{
 		return;
diff --git a/Engine/Plugins/Runtime/InstancedActors/Source/InstancedActors/Public/InstancedActorsData.h b/Engine/Plugins/Runtime/InstancedActors/Source/InstancedActors/Public/InstancedActorsData.h
--- a/Engine/Plugins/Runtime/InstancedActors/Source/InstancedActors/Public/InstancedActorsData.h
+++ b/Engine/Plugins/Runtime/InstancedActors/Source/InstancedActors/Public/InstancedActorsData.h
@@ -502,4 +502,11 @@ private:
 	FGameplayTagContainer CombinedTags;
 
 	TSharedPtr<UE::InstancedActors::FExemplarActorData> ExemplarActorData;
+
+	struct FSetReplicatedActorRequests
+	{
+		FInstancedActorsInstanceIndex InstanceIndex;
+		TWeakObjectPtr<AActor> Actor;
+	};
+	TArray<FSetReplicatedActorRequests> CachedSetReplicatedActorRequests;
 };
