[InstancedActors]  Issue 628300: added a constructor to the newly implemented FSetReplicatedActorRequests type. Required due to Mac compiler not being able to figure this one out on its own.

diff --git a/Engine/Plugins/Runtime/InstancedActors/Source/InstancedActors/Public/InstancedActorsData.h b/Engine/Plugins/Runtime/InstancedActors/Source/InstancedActors/Public/InstancedActorsData.h
--- a/Engine/Plugins/Runtime/InstancedActors/Source/InstancedActors/Public/InstancedActorsData.h
+++ b/Engine/Plugins/Runtime/InstancedActors/Source/InstancedActors/Public/InstancedActorsData.h
@@ -505,6 +505,11 @@ private:
 
 	struct FSetReplicatedActorRequests
 	{
+		FSetReplicatedActorRequests(FInstancedActorsInstanceIndex InInstanceIndex, TWeakObjectPtr<AActor> InActor)
+			: InstanceIndex(InInstanceIndex), Actor(InActor)
+		{	
+		}
+
 		FInstancedActorsInstanceIndex InstanceIndex;
 		TWeakObjectPtr<AActor> Actor;
 	};
