[InstancedActors] Allow UInstancedActorsModifierVolumeComponent::bIgnoreOwnLevelsInstances to check for the situation of externally loaded actors, and check the FTopLevelAssetPath of each levels' World Partition streaming cell

diff --git a/Engine/Plugins/Runtime/InstancedActors/Source/InstancedActors/Private/InstancedActorsModifierVolumeComponent.cpp b/Engine/Plugins/Runtime/InstancedActors/Source/InstancedActors/Private/InstancedActorsModifierVolumeComponent.cpp
--- a/Engine/Plugins/Runtime/InstancedActors/Source/InstancedActors/Private/InstancedActorsModifierVolumeComponent.cpp
+++ b/Engine/Plugins/Runtime/InstancedActors/Source/InstancedActors/Private/InstancedActorsModifierVolumeComponent.cpp
@@ -8,13 +8,36 @@
 #include "InstancedActorsSubsystem.h"
 #include "InstancedActorsTypes.h"
 #include "Engine/World.h"
+#include "Engine/Level.h"
 #include "Algo/NoneOf.h"
 #include "Net/UnrealNetwork.h"
 #include "PrimitiveSceneProxy.h"
 #include "SceneManagement.h"
 #include "ShowFlags.h"
 #include "UObject/Package.h"
+#include "WorldPartition/WorldPartitionRuntimeLevelStreamingCell.h"
+#include "WorldPartition/WorldPartitionLevelStreamingDynamic.h"
 
+namespace UE::InstancedActors
+{
+	namespace Helpers
+	{
+		FTopLevelAssetPath FindLevelWorldPartitionAssetPath(const ULevel* Level)
+		{
+			FTopLevelAssetPath AssetPath;
+			if (const UWorldPartitionRuntimeLevelStreamingCell* WorldPartitionCell = Cast<UWorldPartitionRuntimeLevelStreamingCell>(Level->GetWorldPartitionRuntimeCell()))
+			{
+				if (const UWorldPartitionLevelStreamingDynamic* WorldPartitionLevelStreaming = Cast<UWorldPartitionLevelStreamingDynamic>(WorldPartitionCell->GetLevelStreaming()))
+				{
+					AssetPath = WorldPartitionLevelStreaming->GetOuterWorldPartition().GetAssetPath();
+				}
+			}
+
+			return AssetPath;
+		}
+	} // Helpers
+
+} // namespace InstancedActors
 
 //-----------------------------------------------------------------------------
 // UInstancedActorsModifierVolumeComponent
@@ -141,10 +164,31 @@ bool UInstancedActorsModifierVolumeComponent::TryRunPendingModifiers(AInstancedA
 	}
 
 	check(GetOwner());
-	if (bIgnoreOwnLevelsInstances && Manager.GetLevel() == GetOwner()->GetLevel())
+	if (bIgnoreOwnLevelsInstances)
 	{
-		// Return true as if we have successfully run all modifiers, so they don't attempt to run again.
-		return true;
+		// If logic matches, return true as if we have successfully run all modifiers, so they don't attempt to run again.
+		ULevel* ManagerLevel = Manager.GetLevel();
+		ULevel* OwnerLevel = GetOwner()->GetLevel();
+
+		// If Actors are externally loaded, the ULevel pointers will not match even if they originally came from the same level.
+		// Instead, compare the FTopLevelAssetPath from their World Partition streaming cells
+		if (ManagerLevel != nullptr && OwnerLevel != nullptr && Manager.IsPackageExternal() && GetOwner()->IsPackageExternal())
+		{
+			FTopLevelAssetPath ManagerLoadedPath = UE::InstancedActors::Helpers::FindLevelWorldPartitionAssetPath(ManagerLevel);
+			FTopLevelAssetPath OwnerLoadedPath = UE::InstancedActors::Helpers::FindLevelWorldPartitionAssetPath(OwnerLevel);
+
+			if (ManagerLoadedPath.IsValid() && OwnerLoadedPath.IsValid() && ManagerLoadedPath == OwnerLoadedPath)
+			{
+				return true;
+			}
+		}
+		else
+		{
+			if (ManagerLevel == OwnerLevel)
+			{
+				return true;
+			}
+		}
 	}
 
 	if (!LevelsToIgnore.IsEmpty())
