[InstancedActors] Slightly modify modifier volumes' world partitioned Actor/Level check. Checking for IsPackageExternal() doesn't work for cooked builds. Just do the level asset FName check if it's available for both Actors (which would happen if they're both world partitioned anyway). It would still fallback to the raw pointer check if that failed

diff --git a/Engine/Plugins/Runtime/InstancedActors/Source/InstancedActors/Private/InstancedActorsModifierVolumeComponent.cpp b/Engine/Plugins/Runtime/InstancedActors/Source/InstancedActors/Private/InstancedActorsModifierVolumeComponent.cpp
--- a/Engine/Plugins/Runtime/InstancedActors/Source/InstancedActors/Private/InstancedActorsModifierVolumeComponent.cpp
+++ b/Engine/Plugins/Runtime/InstancedActors/Source/InstancedActors/Private/InstancedActorsModifierVolumeComponent.cpp
@@ -170,9 +170,9 @@ bool UInstancedActorsModifierVolumeComponent::TryRunPendingModifiers(AInstancedA
 		ULevel* ManagerLevel = Manager.GetLevel();
 		ULevel* OwnerLevel = GetOwner()->GetLevel();
 
-		// If Actors are externally loaded, the ULevel pointers will not match even if they originally came from the same level.
+		// If Actors are world partitioned, the ULevel pointers will not match even if they originally came from the same level.
 		// Instead, compare the FTopLevelAssetPath from their World Partition streaming cells
-		if (ManagerLevel != nullptr && OwnerLevel != nullptr && Manager.IsPackageExternal() && GetOwner()->IsPackageExternal())
+		if (ManagerLevel != nullptr && OwnerLevel != nullptr)
 		{
 			FTopLevelAssetPath ManagerLoadedPath = UE::InstancedActors::Helpers::FindLevelWorldPartitionAssetPath(ManagerLevel);
 			FTopLevelAssetPath OwnerLoadedPath = UE::InstancedActors::Helpers::FindLevelWorldPartitionAssetPath(OwnerLevel);
@@ -182,12 +182,10 @@ bool UInstancedActorsModifierVolumeComponent::TryRunPendingModifiers(AInstancedA
 				return true;
 			}
 		}
-		else
+
+		if (ManagerLevel == OwnerLevel)
 		{
-			if (ManagerLevel == OwnerLevel)
-			{
-				return true;
-			}
+			return true;
 		}
 	}
 
