[InstancedActors] Fix AInstancedActorsManager::GetStreamingBounds(). Before it was always returning a WORLD_MAX size, which would ultimately change both its streaming cell and its actor path anytime it got touched/saved

diff --git a/Engine/Plugins/Runtime/InstancedActors/Source/InstancedActors/Private/InstancedActorsManager.cpp b/Engine/Plugins/Runtime/InstancedActors/Source/InstancedActors/Private/InstancedActorsManager.cpp
--- a/Engine/Plugins/Runtime/InstancedActors/Source/InstancedActors/Private/InstancedActorsManager.cpp
+++ b/Engine/Plugins/Runtime/InstancedActors/Source/InstancedActors/Private/InstancedActorsManager.cpp
@@ -652,8 +652,10 @@ void AInstancedActorsManager::GetStreamingBounds(FBox& OutRuntimeBounds, FBox& O
 		UActorPartitionSubsystem::FCellCoord CellCoord = UActorPartitionSubsystem::FCellCoord::GetCellCoord(GetActorLocation(), GetLevel(), GetGridSize());
 		OutRuntimeBounds = OutEditorBounds = UActorPartitionSubsystem::FCellCoord::GetCellBounds(CellCoord, GetGridSize());
 	}
-
-	OutRuntimeBounds = OutEditorBounds = FBox(FVector(-HALF_WORLD_MAX), FVector(HALF_WORLD_MAX));
+	else
+	{
+		OutRuntimeBounds = OutEditorBounds = FBox(FVector(-HALF_WORLD_MAX), FVector(HALF_WORLD_MAX));
+	}
 }
 
 uint32 AInstancedActorsManager::GetDefaultGridSize(UWorld* InWorld) const
