[InstancedActors] Add a small comment clarifying the logic around the world partitioned vs. non-world partitioned level checking

diff --git a/Engine/Plugins/Runtime/InstancedActors/Source/InstancedActors/Private/InstancedActorsModifierVolumeComponent.cpp b/Engine/Plugins/Runtime/InstancedActors/Source/InstancedActors/Private/InstancedActorsModifierVolumeComponent.cpp
--- a/Engine/Plugins/Runtime/InstancedActors/Source/InstancedActors/Private/InstancedActorsModifierVolumeComponent.cpp
+++ b/Engine/Plugins/Runtime/InstancedActors/Source/InstancedActors/Private/InstancedActorsModifierVolumeComponent.cpp
@@ -171,7 +171,9 @@ bool UInstancedActorsModifierVolumeComponent::TryRunPendingModifiers(AInstancedA
 		ULevel* OwnerLevel = GetOwner()->GetLevel();
 
 		// If Actors are world partitioned, the ULevel pointers will not match even if they originally came from the same level.
-		// Instead, compare the FTopLevelAssetPath from their World Partition streaming cells
+		// Instead, compare the FTopLevelAssetPath from their World Partition streaming cells if applicable.
+		// For non-world partitioned Actors, the logic will fall to the next block to do the original ULevel pointer comparison,
+		// since either ManagerLoadedPath or OwnerLoadedPath will remain !IsValid()
 		if (ManagerLevel != nullptr && OwnerLevel != nullptr)
 		{
 			FTopLevelAssetPath ManagerLoadedPath = UE::InstancedActors::Helpers::FindLevelWorldPartitionAssetPath(ManagerLevel);
