[InstancedActors] Massaged UInstancedActorsVisualizationSwitcherProcessor::SwitchEntityMeshDesc a bit to get us more information when data isn't wright (doing it via ensures, to not miss it).

#rnx

diff --git a/Engine/Plugins/Runtime/InstancedActors/Source/InstancedActors/Private/InstancedActorsVisualizationSwitcherProcessor.cpp b/Engine/Plugins/Runtime/InstancedActors/Source/InstancedActors/Private/InstancedActorsVisualizationSwitcherProcessor.cpp
--- a/Engine/Plugins/Runtime/InstancedActors/Source/InstancedActors/Private/InstancedActorsVisualizationSwitcherProcessor.cpp
+++ b/Engine/Plugins/Runtime/InstancedActors/Source/InstancedActors/Private/InstancedActorsVisualizationSwitcherProcessor.cpp
@@ -64,7 +64,9 @@ void UInstancedActorsVisualizationSwitcherProcessor::SwitchEntityMeshDesc(FMassI
 		// and 'forget' about it.
 		if (RepresentationFragment.PrevRepresentation == EMassRepresentationType::StaticMeshInstance)
 		{
-			if (RepresentationFragment.StaticMeshDescHandle.IsValid())
+			if (ensureMsgf(RepresentationFragment.StaticMeshDescHandle.IsValid(), TEXT("Switching visualization while StaticMeshDescHandle is invalid"))
+				&& ensureMsgf(ISMInfosView.IsValidIndex(RepresentationFragment.StaticMeshDescHandle.ToIndex()), TEXT("Switching visualization while StaticMeshDescHandle is not a valid ISMInfo index, %d out of %d")
+					, RepresentationFragment.StaticMeshDescHandle.ToIndex(), ISMInfosView.Num()))
 			{
 				FMassInstancedStaticMeshInfo& ISMInfo = ISMInfosView[RepresentationFragment.StaticMeshDescHandle.ToIndex()];
 
diff --git a/Engine/Plugins/Runtime/MassGameplay/Source/MassRepresentation/Public/MassRepresentationTypes.h b/Engine/Plugins/Runtime/MassGameplay/Source/MassRepresentation/Public/MassRepresentationTypes.h
--- a/Engine/Plugins/Runtime/MassGameplay/Source/MassRepresentation/Public/MassRepresentationTypes.h
+++ b/Engine/Plugins/Runtime/MassGameplay/Source/MassRepresentation/Public/MassRepresentationTypes.h
@@ -806,6 +806,11 @@ struct FMassInstancedStaticMeshInfoArrayViewAccessDetector
 		return InstancedStaticMeshInfos.IsValidIndex(Index);
 	}
 
+	int32 Num() const
+	{
+		return InstancedStaticMeshInfos.Num();
+	}
+
 private:
 	TArrayView<FMassInstancedStaticMeshInfo> InstancedStaticMeshInfos;
 	const FRWRecursiveAccessDetector* AccessDetector;
