[InstancedActors] Fixed inconsistencies in how IAManager editor-preview ISM components are configured.

#rnx

diff --git a/Engine/Plugins/Runtime/InstancedActors/Source/InstancedActors/Private/InstancedActorsManager.cpp b/Engine/Plugins/Runtime/InstancedActors/Source/InstancedActors/Private/InstancedActorsManager.cpp
--- a/Engine/Plugins/Runtime/InstancedActors/Source/InstancedActors/Private/InstancedActorsManager.cpp
+++ b/Engine/Plugins/Runtime/InstancedActors/Source/InstancedActors/Private/InstancedActorsManager.cpp
@@ -749,10 +749,7 @@ void AInstancedActorsManager::PreRegisterAllComponents()
 		// Modify the ISMCs here since the components might have been serialized before this change.
 		for (TObjectPtr<UInstancedStaticMeshComponent> ISMComponent : InstanceData->EditorPreviewISMComponents)
 		{
-			ISMComponent->SetIsVisualizationComponent(true);
-			ISMComponent->SetCollisionEnabled(ECollisionEnabled::QueryOnly);
-			ISMComponent->bSelectable = true;
-			ISMComponent->bHasPerInstanceHitProxies = true;
+			SetUpEditorPreviewISMComponent(ISMComponent);
 		}
 
 		// Add EditorPreviewISMComponents to the known ISMComponent-to-InstanceData map (ISMComponentToInstanceDataMap)
@@ -1564,12 +1561,7 @@ void AInstancedActorsManager::CreateISMComponents(const FInstancedActorsVisualiz
 		// Editor preview overrides
 		if (bEditorPreviewISMCs)
 		{
-#if WITH_EDITORONLY_DATA
-			ISMComponent->SetIsVisualizationComponent(true);
-#endif
-			ISMComponent->SetCollisionEnabled(ECollisionEnabled::NoCollision);
-			ISMComponent->bSelectable = true;
-			ISMComponent->bHasPerInstanceHitProxies = true;
+			SetUpEditorPreviewISMComponent(ISMComponent);
 		}
 
 		ISMComponent->RegisterComponent();
@@ -1578,6 +1570,16 @@ void AInstancedActorsManager::CreateISMComponents(const FInstancedActorsVisualiz
 	}
 }
 
+void AInstancedActorsManager::SetUpEditorPreviewISMComponent(gsl::not_null<UInstancedStaticMeshComponent*> ISMComponent)
+{
+#if WITH_EDITORONLY_DATA
+	ISMComponent->SetIsVisualizationComponent(true);
+#endif //WITH_EDITORONLY_DATA
+	ISMComponent->SetCollisionEnabled(ECollisionEnabled::QueryOnly);
+	ISMComponent->bSelectable = true;
+	ISMComponent->bHasPerInstanceHitProxies = true;
+}
+
 //-----------------------------------------------------------------------------
 // AInstancedActorsManager IActorInstanceManagerInterface overrides
 //-----------------------------------------------------------------------------
diff --git a/Engine/Plugins/Runtime/InstancedActors/Source/InstancedActors/Public/InstancedActorsManager.h b/Engine/Plugins/Runtime/InstancedActors/Source/InstancedActors/Public/InstancedActorsManager.h
--- a/Engine/Plugins/Runtime/InstancedActors/Source/InstancedActors/Public/InstancedActorsManager.h
+++ b/Engine/Plugins/Runtime/InstancedActors/Source/InstancedActors/Public/InstancedActorsManager.h
@@ -300,6 +300,9 @@ protected:
 	virtual void PreRegisterAllComponents() override;
 #endif
 
+	/** Configures given ISMComponent for editor-preview purposes */
+	static void SetUpEditorPreviewISMComponent(gsl::not_null<UInstancedStaticMeshComponent*> ISMComponent);
+
 	UPROPERTY(Transient)
 	TObjectPtr<UInstancedActorsSubsystem> InstancedActorSubsystem;
 
