[InstancedActors] modified how we report failure of representation validity check in UInstancedActorsStationaryLODBatchProcessor::Execute.

diff --git a/Engine/Plugins/Runtime/InstancedActors/Source/InstancedActors/Private/InstancedActorsStationaryLODBatchProcessor.cpp b/Engine/Plugins/Runtime/InstancedActors/Source/InstancedActors/Private/InstancedActorsStationaryLODBatchProcessor.cpp
--- a/Engine/Plugins/Runtime/InstancedActors/Source/InstancedActors/Private/InstancedActorsStationaryLODBatchProcessor.cpp
+++ b/Engine/Plugins/Runtime/InstancedActors/Source/InstancedActors/Private/InstancedActorsStationaryLODBatchProcessor.cpp
@@ -417,7 +417,7 @@ void UInstancedActorsStationaryLODBatchProcessor::Execute(FMassEntityManager& En
 			{
 				// Note that it's possible for DirtyRepresentationInstances to contain indices to entities that just have 
 				// been destroyed (however, we only expect it to happen only on clients, of if we're somewhere in the process of destroying the world). 
-				ensureMsgf(DirtyRepresentationInstance.IsValid() || bInvalidInstancedAllowed
+				UE_CLOG(!(DirtyRepresentationInstance.IsValid() || bInvalidInstancedAllowed), LogInstancedActors, Warning
 					, TEXT("We only expect invalid instance handles on Client or when the InstancedActorSubsystem no longer has a valid outer UWorld."));
 
 				if (DirtyRepresentationInstance.IsValid() 
