[InstancedActors] Fixed not being able to pick InstancedActors instanced for editing in the editor

diff --git a/Engine/Plugins/Runtime/InstancedActors/Source/InstancedActors/Private/InstancedActorsManager.cpp b/Engine/Plugins/Runtime/InstancedActors/Source/InstancedActors/Private/InstancedActorsManager.cpp
--- a/Engine/Plugins/Runtime/InstancedActors/Source/InstancedActors/Private/InstancedActorsManager.cpp
+++ b/Engine/Plugins/Runtime/InstancedActors/Source/InstancedActors/Private/InstancedActorsManager.cpp
@@ -1572,9 +1572,6 @@ void AInstancedActorsManager::CreateISMComponents(const FInstancedActorsVisualiz
 
 void AInstancedActorsManager::SetUpEditorPreviewISMComponent(gsl::not_null<UInstancedStaticMeshComponent*> ISMComponent)
 {
-#if WITH_EDITORONLY_DATA
-	ISMComponent->SetIsVisualizationComponent(true);
-#endif //WITH_EDITORONLY_DATA
 	ISMComponent->SetCollisionEnabled(ECollisionEnabled::QueryOnly);
 	ISMComponent->bSelectable = true;
 	ISMComponent->bHasPerInstanceHitProxies = true;
