[InstancedActors] Made instanced with settings of MaxActorDistance == 0 never hydrate, including as a response to physics queries. This makes never-hydrating InstancedActors a lot cheaper

The new feature is disabled by default and is controlled by IA.EnableCanHydrateLogic console variable.

diff --git a/Engine/Plugins/Runtime/InstancedActors/Source/InstancedActors/Private/InstancedActorsData.cpp b/Engine/Plugins/Runtime/InstancedActors/Source/InstancedActors/Private/InstancedActorsData.cpp
--- a/Engine/Plugins/Runtime/InstancedActors/Source/InstancedActors/Private/InstancedActorsData.cpp
+++ b/Engine/Plugins/Runtime/InstancedActors/Source/InstancedActors/Private/InstancedActorsData.cpp
@@ -160,6 +160,13 @@ namespace UE::InstancedActors
 			bEnableReleasingEntityTemplatesAndExemplarActors,
 			TEXT("If enabled (default) instanced actors will release created entity templates and exemplar actors from central systems. Allows resources that are no longer in use to be garbage collected."),
 			ECVF_Default);
+
+		bool bCanHydrateLogicEnabled = true;
+		FAutoConsoleVariableRef CVarCanHydrateLogicEnabled(
+			TEXT("IA.EnableCanHydrateLogic"),
+			bCanHydrateLogicEnabled,
+			TEXT("Toggles special handling of instanced marked as non-hydrating via settings (MaxActorDistance ==0)"),
+			ECVF_Default);
 	} // CVars
 
 	namespace Helpers
@@ -265,6 +272,8 @@ void UInstancedActorsData::Initialize()
 	const uint8 VisualizationIndex = AddVisualization(DefaultVisualiation);
 	check(VisualizationIndex == 0);
 
+	bCanHydrate = (UE::InstancedActors::CVars::bCanHydrateLogicEnabled == false) || (Settings == nullptr) || Settings->MaxActorDistance > 0;
+
 	// Create entity template
 	CreateEntityTemplate(*ExemplarActor);
 
@@ -312,6 +321,11 @@ void UInstancedActorsData::CreateEntityTemplate(const AActor& ExemplarActor)
 	const FMassEntityTemplate& BaseEntityTemplate = EntityConfig.GetOrCreateEntityTemplate(*World);
 
 	FMassEntityTemplateData ModifiedTemplate(BaseEntityTemplate);
+	if (bCanHydrate)
+	{
+		ModifiedTemplate.AddTag<FInstancedActorsHydratableTag>();
+	}
+
 	ModifyEntityTemplate(ModifiedTemplate, ExemplarActor);
 	
 	// Allow UInstancedActorsComponent's to extend entity template
diff --git a/Engine/Plugins/Runtime/InstancedActors/Source/InstancedActors/Private/InstancedActorsManager.cpp b/Engine/Plugins/Runtime/InstancedActors/Source/InstancedActors/Private/InstancedActorsManager.cpp
--- a/Engine/Plugins/Runtime/InstancedActors/Source/InstancedActors/Private/InstancedActorsManager.cpp
+++ b/Engine/Plugins/Runtime/InstancedActors/Source/InstancedActors/Private/InstancedActorsManager.cpp
@@ -1657,28 +1657,31 @@ AActor* AInstancedActorsManager::FindActorInternal(const FActorInstanceHandle& H
 		const UInstancedActorsData* InstanceData = PerActorClassInstanceData[InstancedActorDataIndex];
 		check(IsValid(InstanceData));
 
-		const FMassEntityHandle EntityHandle = InstanceData->GetEntityHandleForIndex(InternalInstanceIndex);
-		// note that it's possible that InternalInstanceIndex points at a no-longer-valid entity, if
-		// the entity got already destroyed due to other responses to the physics query that
-		// generated the given FActorInstanceHandle handle. This can easily be reproduced if actors
-		// get destroyed as a result of being hit.
-		if (EntityHandle.IsValid())
+		if (InstanceData->CanHydrate())
 		{
-			const FMassEntityManager& EntityManager = GetMassEntityManagerChecked();
-			OutEntityView = FMassEntityView::TryMakeView(EntityManager, EntityHandle);
-			if (OutEntityView.IsValid())
+			const FMassEntityHandle EntityHandle = InstanceData->GetEntityHandleForIndex(InternalInstanceIndex);
+			// note that it's possible that InternalInstanceIndex points at a no-longer-valid entity, if
+			// the entity got already destroyed due to other responses to the physics query that
+			// generated the given FActorInstanceHandle handle. This can easily be reproduced if actors
+			// get destroyed as a result of being hit.
+			if (EntityHandle.IsValid())
 			{
-				if (FMassActorFragment* ActorFragment = OutEntityView.GetFragmentDataPtr<FMassActorFragment>())
+				const FMassEntityManager& EntityManager = GetMassEntityManagerChecked();
+				OutEntityView = FMassEntityView::TryMakeView(EntityManager, EntityHandle);
+				if (OutEntityView.IsValid())
 				{
-					return ActorFragment->GetMutable();
+					if (FMassActorFragment* ActorFragment = OutEntityView.GetFragmentDataPtr<FMassActorFragment>())
+					{
+						return ActorFragment->GetMutable();
+					}
 				}
-			}
 #if DO_ENSURE
-			else
-			{
-				ensureMsgf(!bEnsureOnMissingInstanceDataOrMassEntity, TEXT("Unable to fetch a valid FMassEntityView for entity handle [%s]"), *EntityHandle.DebugGetDescription());
-			}
+				else
+				{
+					ensureMsgf(!bEnsureOnMissingInstanceDataOrMassEntity, TEXT("Unable to fetch a valid FMassEntityView for entity handle [%s]"), *EntityHandle.DebugGetDescription());
+				}
 #endif // DO_ENSURE
+			}
 		}
 	}
 #if DO_ENSURE
@@ -1720,6 +1723,18 @@ AActor* AInstancedActorsManager::FindOrCreateActor(const FActorInstanceHandle& H
 	// Create missing actor from Mass if we have a valid EntityView
 	if (Actor == nullptr && EntityView.IsValid())
 	{
+#if WITH_INSTANCEDACTORS_DEBUG
+		{
+			const int32 CompositeIndex = Handle.GetInstanceIndex();
+			const int32 InstancedActorDataIndex = FInstancedActorsInstanceIndex::ExtractInstanceDataID(CompositeIndex);
+			const FInstancedActorsInstanceIndex InternalInstanceIndex = FInstancedActorsInstanceIndex(FInstancedActorsInstanceIndex::ExtractInternalInstanceIndex(CompositeIndex));
+			const UInstancedActorsData* InstanceData = PerActorClassInstanceData.IsValidIndex(InstancedActorDataIndex) ? PerActorClassInstanceData[InstancedActorDataIndex] : nullptr;
+			ensureMsgf(InstanceData && InstanceData->CanHydrate()
+				, TEXT("We're about to spawn an actor while the relevant IAD %s has been configured to not hydrate its instances")
+				, InstanceData ? *InstanceData->GetDebugName() : TEXT("NONE"));
+		}
+#endif // WITH_INSTANCEDACTORS_DEBUG
+
 		UMassRepresentationSubsystem* RepresentationSubsystem = EntityView.GetSharedFragmentData<FMassRepresentationSubsystemSharedFragment>().RepresentationSubsystem;
 		const FMassRepresentationParameters& RepresentationParams = EntityView.GetConstSharedFragmentData<FMassRepresentationParameters>();
 		UInstancedActorsRepresentationActorManagement* RepresentationActorManagement = Cast<UInstancedActorsRepresentationActorManagement>(RepresentationParams.CachedRepresentationActorManagement);
diff --git a/Engine/Plugins/Runtime/InstancedActors/Source/InstancedActors/Private/InstancedActorsStationaryLODBatchProcessor.cpp b/Engine/Plugins/Runtime/InstancedActors/Source/InstancedActors/Private/InstancedActorsStationaryLODBatchProcessor.cpp
--- a/Engine/Plugins/Runtime/InstancedActors/Source/InstancedActors/Private/InstancedActorsStationaryLODBatchProcessor.cpp
+++ b/Engine/Plugins/Runtime/InstancedActors/Source/InstancedActors/Private/InstancedActorsStationaryLODBatchProcessor.cpp
@@ -69,6 +69,40 @@ static FAutoConsoleVariableRef AnonymousCVars[] = {
 
 } // UE::Mass::Tweakables
 
+namespace UE::InstancedActors
+{
+	bool EnablePhysicForVisualization(uint8 /*VisualizationIndex*/, const FInstancedActorsVisualizationInfo& Visualization)
+	{
+		for (const TObjectPtr<UInstancedStaticMeshComponent>& ISMComponent : Visualization.ISMComponents)
+		{
+			if (ISMComponent)
+			{
+				if (ISMComponent->IsRegistered())
+				{
+					ISMComponent->CreatePhysicsState(/*bAllowDeferral=*/true);
+				}
+				else
+				{
+					UE_LOG(LogInstancedActors, Error, TEXT("Failed to call CreatePhysicsState() on component '%s', because component is not registered."), *GetFullNameSafe(ISMComponent));
+				}
+			}
+		}
+		return true;
+	}
+
+	bool DisablePhysicForVisualization(uint8 /*VisualizationIndex*/, const FInstancedActorsVisualizationInfo& Visualization)
+	{
+		for (const TObjectPtr<UInstancedStaticMeshComponent>& ISMComponent : Visualization.ISMComponents)
+		{
+			if (ISMComponent)
+			{
+				ISMComponent->DestroyPhysicsState();
+			}
+		}
+		return true;
+	}
+}
+
 //-----------------------------------------------------------------------------
 // UInstancedActorsStationaryLODBatchProcessor
 //-----------------------------------------------------------------------------
@@ -246,10 +280,11 @@ void UInstancedActorsStationaryLODBatchProcessor::Execute(FMassEntityManager& En
 						NewBulkLOD = EInstancedActorsBulkLOD::Low;
 					}
 
-					static EInstancedActorsBulkLOD BulkLODOverride = EInstancedActorsBulkLOD::MAX;
-					if (BulkLODOverride != EInstancedActorsBulkLOD::MAX)
+					EInstancedActorsBulkLOD BulkLODOverriddenByNonHydrating = EInstancedActorsBulkLOD::MAX;
+					if (NewBulkLOD == EInstancedActorsBulkLOD::Detailed && InstanceData->CanHydrate() == false)
 					{
-						NewBulkLOD = BulkLODOverride;
+						BulkLODOverriddenByNonHydrating = NewBulkLOD;
+						NewBulkLOD = EInstancedActorsBulkLOD::Medium;
 					}
 
 					check(NewBulkLOD != EInstancedActorsBulkLOD::MAX);
@@ -264,48 +299,19 @@ void UInstancedActorsStationaryLODBatchProcessor::Execute(FMassEntityManager& En
 							ManagerSharedFragment.BulkLOD = NewBulkLOD;
 							AInstancedActorsManager::UpdateInstanceStats(InstanceData->NumInstances, ManagerSharedFragment.BulkLOD, true);
 						}
-
 						// Toggles physics state for the IA's ISM depending on the new bulk LOD value.
 						// If enabled = physics on, else = physics off.				
 						if (UE::Mass::Tweakables::bControlPhysicsState && Settings.bControlPhysicsState)
 						{
-							if (NewBulkLOD == EInstancedActorsBulkLOD::Detailed)
+							if (NewBulkLOD == EInstancedActorsBulkLOD::Detailed || BulkLODOverriddenByNonHydrating == EInstancedActorsBulkLOD::Detailed)
 							{
-								InstanceData->ForEachVisualization([](uint8 VisualizationIndex, const FInstancedActorsVisualizationInfo& Visualization)
-								{
-									for (const TObjectPtr<UInstancedStaticMeshComponent>& ISMComponent : Visualization.ISMComponents)
-									{
-										if (ISMComponent)
-										{
-											if (ISMComponent->IsRegistered())
-											{
-												ISMComponent->CreatePhysicsState(/*bAllowDeferral=*/true);
-											}
-											else
-											{
-												UE_LOG(LogInstancedActors, Error, TEXT("Failed to call CreatePhysicsState() on component '%s', because component is not registered."), *GetFullNameSafe(ISMComponent));
-											}
-										}
-									}
-									return true;
-								});
+								InstanceData->ForEachVisualization(&UE::InstancedActors::EnablePhysicForVisualization);
 							}
 							else
 							{
-								InstanceData->ForEachVisualization([](uint8 VisualizationIndex, const FInstancedActorsVisualizationInfo& Visualization)
-								{
-									for (const TObjectPtr<UInstancedStaticMeshComponent>& ISMComponent : Visualization.ISMComponents)
-									{
-										if (ISMComponent)
-										{
-											ISMComponent->DestroyPhysicsState();
-										}
-									}
-									return true;
-								});
+								InstanceData->ForEachVisualization(UE::InstancedActors::DisablePhysicForVisualization);
 							}
 						}
-
 						{
 							// Toggles visibility for the IA's ISM depending on the new bulk LOD value.
 							// If enabled = use default visibility (probably on), else = physics off.
diff --git a/Engine/Plugins/Runtime/InstancedActors/Source/InstancedActors/Private/InstancedActorsVisualizationProcessor.cpp b/Engine/Plugins/Runtime/InstancedActors/Source/InstancedActors/Private/InstancedActorsVisualizationProcessor.cpp
--- a/Engine/Plugins/Runtime/InstancedActors/Source/InstancedActors/Private/InstancedActorsVisualizationProcessor.cpp
+++ b/Engine/Plugins/Runtime/InstancedActors/Source/InstancedActors/Private/InstancedActorsVisualizationProcessor.cpp
@@ -2,6 +2,7 @@
 
 #include "InstancedActorsVisualizationProcessor.h"
 #include "MassStationaryISMSwitcherProcessor.h"
+#include "InstancedActorsTypes.h"
 
 
 UInstancedActorsVisualizationProcessor::UInstancedActorsVisualizationProcessor()
@@ -22,5 +23,6 @@ void UInstancedActorsVisualizationProcessor::ConfigureQueries(const TSharedRef<F
 	Super::ConfigureQueries(EntityManager);
 
 	EntityQuery.ClearTagRequirements(FMassTagBitSet(*FMassVisualizationProcessorTag::StaticStruct()));
+	EntityQuery.AddTagRequirement<FInstancedActorsHydratableTag>(EMassFragmentPresence::All);
 	EntityQuery.AddTagRequirement<FInstancedActorsVisualizationProcessorTag>(EMassFragmentPresence::All);
 }
\ No newline at end of file
diff --git a/Engine/Plugins/Runtime/InstancedActors/Source/InstancedActors/Public/InstancedActorsData.h b/Engine/Plugins/Runtime/InstancedActors/Source/InstancedActors/Public/InstancedActorsData.h
--- a/Engine/Plugins/Runtime/InstancedActors/Source/InstancedActors/Public/InstancedActorsData.h
+++ b/Engine/Plugins/Runtime/InstancedActors/Source/InstancedActors/Public/InstancedActorsData.h
@@ -185,6 +185,8 @@ public:
 	// Returns true if InstanceTransforms has been consumed to spawn Mass entities
 	bool HasSpawnedEntities() const;
 
+	bool CanHydrate() const;
+
 	// Returns the total instance count, including both valid & invalid instances e.g: GetNumFreeInstances() + NumValidInstances
 	// Useful for iterating by FInstancedActorsInstanceIndex or allocating same-sized instance data arrays
 	int32 GetNumInstances() const;
@@ -501,6 +503,8 @@ private:
 	UPROPERTY(Transient)
 	FGameplayTagContainer CombinedTags;
 
+	bool bCanHydrate = true;
+
 	TSharedPtr<UE::InstancedActors::FExemplarActorData> ExemplarActorData;
 
 	struct FSetReplicatedActorRequests
@@ -515,3 +519,11 @@ private:
 	};
 	TArray<FSetReplicatedActorRequests> CachedSetReplicatedActorRequests;
 };
+
+//-----------------------------------------------------------------------------
+// INLINES
+//-----------------------------------------------------------------------------
+inline bool UInstancedActorsData::CanHydrate() const
+{
+	return bCanHydrate;
+}
\ No newline at end of file
diff --git a/Engine/Plugins/Runtime/InstancedActors/Source/InstancedActors/Public/InstancedActorsTypes.h b/Engine/Plugins/Runtime/InstancedActors/Source/InstancedActors/Public/InstancedActorsTypes.h
--- a/Engine/Plugins/Runtime/InstancedActors/Source/InstancedActors/Public/InstancedActorsTypes.h
+++ b/Engine/Plugins/Runtime/InstancedActors/Source/InstancedActors/Public/InstancedActorsTypes.h
@@ -323,3 +323,9 @@ struct FInstancedActorsFragment : public FMassFragment
 	UPROPERTY()
 	FInstancedActorsInstanceIndex InstanceIndex;
 };
+
+USTRUCT()
+struct INSTANCEDACTORS_API FInstancedActorsHydratableTag : public FMassTag
+{
+	GENERATED_BODY();
+};
