[InstancedActors] a fix to the change introducing non-hydrating instances. The previous worked only in some circumstances (like teleporting close to non-hydrating IAs).

diff --git a/Engine/Plugins/Runtime/InstancedActors/Source/InstancedActors/Private/InstancedActorsStationaryLODBatchProcessor.cpp b/Engine/Plugins/Runtime/InstancedActors/Source/InstancedActors/Private/InstancedActorsStationaryLODBatchProcessor.cpp
--- a/Engine/Plugins/Runtime/InstancedActors/Source/InstancedActors/Private/InstancedActorsStationaryLODBatchProcessor.cpp
+++ b/Engine/Plugins/Runtime/InstancedActors/Source/InstancedActors/Private/InstancedActorsStationaryLODBatchProcessor.cpp
@@ -280,13 +280,6 @@ void UInstancedActorsStationaryLODBatchProcessor::Execute(FMassEntityManager& En
 						NewBulkLOD = EInstancedActorsBulkLOD::Low;
 					}
 
-					EInstancedActorsBulkLOD BulkLODOverriddenByNonHydrating = EInstancedActorsBulkLOD::MAX;
-					if (NewBulkLOD == EInstancedActorsBulkLOD::Detailed && InstanceData->CanHydrate() == false)
-					{
-						BulkLODOverriddenByNonHydrating = NewBulkLOD;
-						NewBulkLOD = EInstancedActorsBulkLOD::Medium;
-					}
-
 					check(NewBulkLOD != EInstancedActorsBulkLOD::MAX);
 					// Updates the time at which the FInstancedActorsDataSharedFragment will tick depending on its bulk LOD value
 					NextTickTime = CurrentTime + (DelayPerBulkLOD[(int)NewBulkLOD] * 0.95 + FMath::FRand() * 0.1);
@@ -303,7 +296,7 @@ void UInstancedActorsStationaryLODBatchProcessor::Execute(FMassEntityManager& En
 						// If enabled = physics on, else = physics off.				
 						if (UE::Mass::Tweakables::bControlPhysicsState && Settings.bControlPhysicsState)
 						{
-							if (NewBulkLOD == EInstancedActorsBulkLOD::Detailed || BulkLODOverriddenByNonHydrating == EInstancedActorsBulkLOD::Detailed)
+							if (NewBulkLOD == EInstancedActorsBulkLOD::Detailed)
 							{
 								InstanceData->ForEachVisualization(&UE::InstancedActors::EnablePhysicForVisualization);
 							}
@@ -346,7 +339,7 @@ void UInstancedActorsStationaryLODBatchProcessor::Execute(FMassEntityManager& En
 						}
 						// Toggles MassProcessors on/off depending on the bulk LOD, by pushing or removing tags that are used by those processor's queries.
 						// NOTE: Forcibly updates the mass LOD to off or low when bulk LOD is smaller than Detailed
-						if (ManagerSharedFragment.BulkLOD == EInstancedActorsBulkLOD::Detailed)
+						if (ManagerSharedFragment.BulkLOD == EInstancedActorsBulkLOD::Detailed && InstanceData->CanHydrate())
 						{
 							EntityManager.Defer().PushCommand<UE::InstancedActors::FEnableDetailedLODCommand>(InstanceData->Entities);
 						}
@@ -356,6 +349,10 @@ void UInstancedActorsStationaryLODBatchProcessor::Execute(FMassEntityManager& En
 							EMassLOD::Type NewLOD = EMassLOD::Off;
 							switch (ManagerSharedFragment.BulkLOD)
 							{
+							case EInstancedActorsBulkLOD::Detailed:
+								ensureMsgf(InstanceData->CanHydrate() == false, TEXT("This case is only valid for non-hydrating instance, broken for %s"), *GetNameSafe(InstanceData->ActorClass));
+								NewLOD = EMassLOD::Low;
+								break;
 							case EInstancedActorsBulkLOD::Medium: // right now falling through since we don't have a medium-level visualization
 							case EInstancedActorsBulkLOD::Low:
 								NewLOD = EMassLOD::Low;
