[InstancedActors] stopped log spaw from InstancedActors BatchLOD processor when MassLODSubsystem is misconfigured

#rnx

diff --git a/Engine/Plugins/Runtime/InstancedActors/Source/InstancedActors/Private/InstancedActorsStationaryLODBatchProcessor.cpp b/Engine/Plugins/Runtime/InstancedActors/Source/InstancedActors/Private/InstancedActorsStationaryLODBatchProcessor.cpp
--- a/Engine/Plugins/Runtime/InstancedActors/Source/InstancedActors/Private/InstancedActorsStationaryLODBatchProcessor.cpp
+++ b/Engine/Plugins/Runtime/InstancedActors/Source/InstancedActors/Private/InstancedActorsStationaryLODBatchProcessor.cpp
@@ -201,9 +201,7 @@ void UInstancedActorsStationaryLODBatchProcessor::Execute(FMassEntityManager& En
 			// 2. the player hasn't really started yet so there's no pawn and the camera is in its inital location
 			// We need to filter out the latter. 
 			// Note that we rely on UMassSubsystem::bUsePlayerPawnLocationInsteadOfCamera being true here. Without it there's no 
-			// reliable way to differentiate the cases.
-			UE_CLOG(LODSubsystem.IsUsingPlayerPawnLocationInsteadOfCamera() == false
-				, LogInstancedActors, Warning, TEXT("Using Player's camera location for instanced actors LOD calculations - this can skew the LOD calculations in non-FPP games."));
+			// reliable way to differentiate the cases - this property is checked in UInstancedActorsSubsystem::Initialize
 			if (APlayerController* ViewerAsPlayerController = Viewers[ViewerIndex].GetPlayerController())
 			{
 				if (ViewerAsPlayerController->GetPawn() == nullptr)
diff --git a/Engine/Plugins/Runtime/InstancedActors/Source/InstancedActors/Private/InstancedActorsSubsystem.cpp b/Engine/Plugins/Runtime/InstancedActors/Source/InstancedActors/Private/InstancedActorsSubsystem.cpp
--- a/Engine/Plugins/Runtime/InstancedActors/Source/InstancedActors/Private/InstancedActorsSubsystem.cpp
+++ b/Engine/Plugins/Runtime/InstancedActors/Source/InstancedActors/Private/InstancedActorsSubsystem.cpp
@@ -171,7 +171,7 @@ void UInstancedActorsSubsystem::Initialize(FSubsystemCollectionBase& Collection)
 	// Collect existing modifier volumes, calling AInstancedActorsModifierVolume::OnAddedToSubsystem to inform them of latent addition 
 	// to this subsystem
 	//
-	// Note: Modifiers *must* be collected before managers, to ensure managers can then retrieve these modifiersin 
+	// Note: Modifiers *must* be collected before managers, to ensure managers can then retrieve these modifiers in 
 	// AInstancedActorsManager::OnAddedToSubsystem, providing managers an opportunity to run optimized pre-entity-spawning modifiers.
 	for (TActorIterator<AInstancedActorsModifierVolume> ModifierVolumeIt(World); ModifierVolumeIt; ++ModifierVolumeIt)
 	{
@@ -186,12 +186,20 @@ void UInstancedActorsSubsystem::Initialize(FSubsystemCollectionBase& Collection)
 	{
 		AInstancedActorsManager* Manager = *ManagerIt;
 		check(Manager);
-		// we only case about managers that have already begun play and missed their chance to registred in their BeginPlay
+		// we only case about managers that have already begun play and missed their chance to registered in their BeginPlay
 		if (Manager->HasActorBegunPlay())
 		{
 			AddManager(*Manager);
 		}
 	}
+
+	// InstancedActors rely on Mass LOD subsystem, and we expect some specific configuration to work properly
+	const UMassLODSubsystem* LODSubsystem = Collection.InitializeDependency<UMassLODSubsystem>();
+	if (ensureMsgf(LODSubsystem, TEXT("InstancedActors require MassLODSubsystem's existence to function properly")))
+	{
+		UE_CLOG(LODSubsystem->IsUsingPlayerPawnLocationInsteadOfCamera() == false
+			, LogInstancedActors, Warning, TEXT("Using Player's camera location for instanced actors LOD calculations - this can skew the LOD calculations in non-FPP games."));
+	}
 }
 
 void UInstancedActorsSubsystem::Deinitialize()
