[InstancedActors] Lowerd the severity of one of the InstancedActors config testing log lines.

#rnx

diff --git a/Engine/Plugins/Runtime/InstancedActors/Source/InstancedActors/Private/InstancedActorsSubsystem.cpp b/Engine/Plugins/Runtime/InstancedActors/Source/InstancedActors/Private/InstancedActorsSubsystem.cpp
--- a/Engine/Plugins/Runtime/InstancedActors/Source/InstancedActors/Private/InstancedActorsSubsystem.cpp
+++ b/Engine/Plugins/Runtime/InstancedActors/Source/InstancedActors/Private/InstancedActorsSubsystem.cpp
@@ -198,7 +198,7 @@ void UInstancedActorsSubsystem::Initialize(FSubsystemCollectionBase& Collection)
 	if (ensureMsgf(LODSubsystem, TEXT("InstancedActors require MassLODSubsystem's existence to function properly")))
 	{
 		UE_CLOG(LODSubsystem->IsUsingPlayerPawnLocationInsteadOfCamera() == false
-			, LogInstancedActors, Warning, TEXT("Using Player's camera location for instanced actors LOD calculations - this can skew the LOD calculations in non-FPP games."));
+			, LogInstancedActors, Log, TEXT("Using Player's camera location for instanced actors LOD calculations - this can skew the LOD calculations in non-FPP games."));
 	}
 }
 
