[InstancedActors] Removing FInstancedActorsHydratableTag, as it is redundant. UInstancedActorsStationaryLODBatchProcessor::Execute() already manages adding/removing the FInstancedActorsVisualizationProcessorTag based on CanHydrate(), which the UInstancedActorsVisualizationProcessor cares about for its query.

diff --git a/Engine/Plugins/Runtime/InstancedActors/Source/InstancedActors/Private/InstancedActorsData.cpp b/Engine/Plugins/Runtime/InstancedActors/Source/InstancedActors/Private/InstancedActorsData.cpp
--- a/Engine/Plugins/Runtime/InstancedActors/Source/InstancedActors/Private/InstancedActorsData.cpp
+++ b/Engine/Plugins/Runtime/InstancedActors/Source/InstancedActors/Private/InstancedActorsData.cpp
@@ -321,11 +321,6 @@ void UInstancedActorsData::CreateEntityTemplate(const AActor& ExemplarActor)
 	const FMassEntityTemplate& BaseEntityTemplate = EntityConfig.GetOrCreateEntityTemplate(*World);
 
 	FMassEntityTemplateData ModifiedTemplate(BaseEntityTemplate);
-	if (bCanHydrate)
-	{
-		ModifiedTemplate.AddTag<FInstancedActorsHydratableTag>();
-	}
-
 	ModifyEntityTemplate(ModifiedTemplate, ExemplarActor);
 	
 	// Allow UInstancedActorsComponent's to extend entity template
diff --git a/Engine/Plugins/Runtime/InstancedActors/Source/InstancedActors/Private/InstancedActorsVisualizationProcessor.cpp b/Engine/Plugins/Runtime/InstancedActors/Source/InstancedActors/Private/InstancedActorsVisualizationProcessor.cpp
--- a/Engine/Plugins/Runtime/InstancedActors/Source/InstancedActors/Private/InstancedActorsVisualizationProcessor.cpp
+++ b/Engine/Plugins/Runtime/InstancedActors/Source/InstancedActors/Private/InstancedActorsVisualizationProcessor.cpp
@@ -23,6 +23,5 @@ void UInstancedActorsVisualizationProcessor::ConfigureQueries(const TSharedRef<F
 	Super::ConfigureQueries(EntityManager);
 
 	EntityQuery.ClearTagRequirements(FMassTagBitSet(*FMassVisualizationProcessorTag::StaticStruct()));
-	EntityQuery.AddTagRequirement<FInstancedActorsHydratableTag>(EMassFragmentPresence::All);
 	EntityQuery.AddTagRequirement<FInstancedActorsVisualizationProcessorTag>(EMassFragmentPresence::All);
 }
\ No newline at end of file
diff --git a/Engine/Plugins/Runtime/InstancedActors/Source/InstancedActors/Public/InstancedActorsTypes.h b/Engine/Plugins/Runtime/InstancedActors/Source/InstancedActors/Public/InstancedActorsTypes.h
--- a/Engine/Plugins/Runtime/InstancedActors/Source/InstancedActors/Public/InstancedActorsTypes.h
+++ b/Engine/Plugins/Runtime/InstancedActors/Source/InstancedActors/Public/InstancedActorsTypes.h
@@ -322,10 +322,4 @@ struct FInstancedActorsFragment : public FMassFragment
 	// The fixed index of this 'instance' into InstanceData
 	UPROPERTY()
 	FInstancedActorsInstanceIndex InstanceIndex;
-};
-
-USTRUCT()
-struct INSTANCEDACTORS_API FInstancedActorsHydratableTag : public FMassTag
-{
-	GENERATED_BODY();
-};
+};
\ No newline at end of file
