[InstancedActors] More markup on editor-visualization ISM components to make invidigual instances selectable in the editor

diff --git a/Engine/Plugins/Runtime/InstancedActors/Source/InstancedActors/Private/InstancedActorsManager.cpp b/Engine/Plugins/Runtime/InstancedActors/Source/InstancedActors/Private/InstancedActorsManager.cpp
--- a/Engine/Plugins/Runtime/InstancedActors/Source/InstancedActors/Private/InstancedActorsManager.cpp
+++ b/Engine/Plugins/Runtime/InstancedActors/Source/InstancedActors/Private/InstancedActorsManager.cpp
@@ -1572,9 +1572,13 @@ void AInstancedActorsManager::CreateISMComponents(const FInstancedActorsVisualiz
 
 void AInstancedActorsManager::SetUpEditorPreviewISMComponent(TNotNull<UInstancedStaticMeshComponent*> ISMComponent)
 {
+#if WITH_EDITORONLY_DATA
+	ISMComponent->SetIsVisualizationComponent(false);
+#endif //WITH_EDITORONLY_DATA
 	ISMComponent->SetCollisionEnabled(ECollisionEnabled::QueryOnly);
 	ISMComponent->bSelectable = true;
 	ISMComponent->bHasPerInstanceHitProxies = true;
+	ISMComponent->bIsEditorOnly = true;
 }
 
 //-----------------------------------------------------------------------------
