[InstancedActors] Allow for Instanced Actor class settings to define whether the instances will require a fully enclosed bounds check to be considered for modifiers

diff --git a/Engine/Plugins/Runtime/InstancedActors/Source/InstancedActors/Private/InstancedActorsManager.cpp b/Engine/Plugins/Runtime/InstancedActors/Source/InstancedActors/Private/InstancedActorsManager.cpp
--- a/Engine/Plugins/Runtime/InstancedActors/Source/InstancedActors/Private/InstancedActorsManager.cpp
+++ b/Engine/Plugins/Runtime/InstancedActors/Source/InstancedActors/Private/InstancedActorsManager.cpp
@@ -133,6 +133,82 @@ namespace UE::InstancedActors
 		}
 	} // anonymous
 
+	template <>
+	bool PassesBoundsTest<FSphere>(const FSphere& QueryBounds, EBoundsTestType BoundsTestType, const FInstancedActorsInstanceHandle& InstanceHandle, const FTransform& InstanceTransform)
+	{
+		switch (BoundsTestType)
+		{
+		case EBoundsTestType::Intersect:
+			{
+				// Cheap test first
+				if (QueryBounds.IsInside(InstanceTransform.GetLocation()))
+				{
+					return true;
+				}
+				
+				const FBox& InstancedActorBounds = InstanceHandle.GetInstanceActorData()->GetCachedLocalBounds();
+				const FSphere TransformedSphere = QueryBounds.TransformBy(InstanceTransform.Inverse());
+				return FMath::SphereAABBIntersection(TransformedSphere, InstancedActorBounds);
+			}
+			break;
+		case EBoundsTestType::Enclosed:
+			{
+				// Cheap test first
+				if (!QueryBounds.IsInside(InstanceTransform.GetLocation()))
+				{
+					return false;
+				}
+				
+				const FBox& InstancedActorBounds = InstanceHandle.GetInstanceActorData()->GetCachedLocalBounds();
+				const FSphere TransformedSphere = QueryBounds.TransformBy(InstanceTransform.Inverse());
+				return QueryBounds.IsInside(InstancedActorBounds.Min) && QueryBounds.IsInside(InstancedActorBounds.Max);
+			}
+			break;
+		default:
+			ensureMsgf(false, TEXT("Unexpected BoundsTestType: %i"), (int)BoundsTestType);
+			break;
+		}
+
+		return false;
+	}
+
+	template <>
+	bool PassesBoundsTest<FBox>(const FBox& QueryBounds, EBoundsTestType BoundsTestType, const FInstancedActorsInstanceHandle& InstanceHandle, const FTransform& InstanceTransform)
+	{
+		switch (BoundsTestType)
+		{
+		case EBoundsTestType::Intersect:
+			{
+				// Cheap test first
+				if (QueryBounds.IsInside(InstanceTransform.GetLocation()))
+				{
+					return true;
+				}
+				
+				const FBox InstancedActorBounds = InstanceHandle.GetInstanceActorData()->GetCachedLocalBounds().TransformBy(InstanceTransform);
+				return QueryBounds.Intersect(InstancedActorBounds);
+			}
+			break;
+		case EBoundsTestType::Enclosed:
+			{
+				// Cheap test first
+				if (!QueryBounds.IsInside(InstanceTransform.GetLocation()))
+				{
+					return false;
+				}
+				
+				const FBox InstancedActorBounds = InstanceHandle.GetInstanceActorData()->GetCachedLocalBounds().TransformBy(InstanceTransform);
+				return QueryBounds.IsInside(InstancedActorBounds);
+			}
+			break;
+		default:
+			ensureMsgf(false, TEXT("Unexpected BoundsTestType: %i"), (int)BoundsTestType);
+			break;
+		}
+		
+		return false;
+	}
+
 } // namespace InstancedActors
 
 //-----------------------------------------------------------------------------
@@ -864,7 +940,7 @@ bool AInstancedActorsManager::ForEachInstance(FInstanceOperationFunc Operation,
 			UE::Mass::Utils::CreateEntityCollections(*MassEntityManager, InstanceData->Entities, FMassArchetypeEntityCollection::NoDuplicates, EntityCollections);
 
 			FMassExecutionContext ExecutionContext(*MassEntityManager);
-			InstancedActorLocationQuery.ForEachEntityChunkInCollections(EntityCollections, ExecutionContext, [InstanceData, &IterationContext, &InstanceHandle, &Operation, &bContinue](FMassExecutionContext& Context)
+			InstancedActorLocationQuery.ForEachEntityChunkInCollections(EntityCollections, ExecutionContext, [&IterationContext, &InstanceHandle, &Operation, &bContinue](FMassExecutionContext& Context)
 				{
 					if (!bContinue)
 					{
@@ -965,48 +1041,17 @@ bool AInstancedActorsManager::ForEachInstance(const TBoundsType& QueryBounds, FI
 	return ForEachInstance(QueryBounds, Operation, IterationContext);
 }
 
-template <>
-bool AInstancedActorsManager::ForEachInstance<FSphere>(const FSphere& QueryBounds, FInstanceOperationFunc Operation, FInstancedActorsIterationContext& IterationContext, TOptional<FInstancedActorDataPredicateFunc> InstancedActorDataPredicate) const
-{
-	return ForEachInstance([&QueryBounds, &Operation](const FInstancedActorsInstanceHandle& InstanceHandle, const FTransform& InstanceTransform, FInstancedActorsIterationContext& IterationContext)
-		{
-			// Cheap test first.
-			if (QueryBounds.IsInside(InstanceTransform.GetLocation()))
-			{
-				return Operation(InstanceHandle, InstanceTransform, IterationContext);
-			}
-
-			// More expensive bounds test.
-			const FBox& InstancedActorBounds = InstanceHandle.InstancedActorData->CachedLocalBounds;
-			const FSphere TransformedSphere = QueryBounds.TransformBy(InstanceTransform.Inverse());
-			if (FMath::SphereAABBIntersection(TransformedSphere, InstancedActorBounds))
-			{
-				return Operation(InstanceHandle, InstanceTransform, IterationContext);
-			}
-
-			return true; },
-		IterationContext, InstancedActorDataPredicate);
-}
-
-template <>
-bool AInstancedActorsManager::ForEachInstance<FBox>(const FBox& QueryBounds, FInstanceOperationFunc Operation, FInstancedActorsIterationContext& IterationContext, TOptional<FInstancedActorDataPredicateFunc> InstancedActorDataPredicate) const
+template <typename TBoundsType>
+bool AInstancedActorsManager::ForEachInstance(const TBoundsType& QueryBounds, FInstanceOperationFunc Operation, FInstancedActorsIterationContext& IterationContext, TOptional<FInstancedActorDataPredicateFunc> InstancedActorDataPredicate) const
 {
 	return ForEachInstance([&QueryBounds, &Operation](const FInstancedActorsInstanceHandle& InstanceHandle, const FTransform& InstanceTransform, FInstancedActorsIterationContext& IterationContext)
 		{
-			// Cheap test first.
-			if (QueryBounds.IsInside(InstanceTransform.GetLocation()))
+			if (UE::InstancedActors::PassesBoundsTest(QueryBounds, UE::InstancedActors::EBoundsTestType::Intersect, InstanceHandle, InstanceTransform))
 			{
 				return Operation(InstanceHandle, InstanceTransform, IterationContext);
 			}
-
-			// More expensive bounds test.
-			const FBox InstancedActorBounds = InstanceHandle.InstancedActorData->CachedLocalBounds.TransformBy(InstanceTransform);
-			if (QueryBounds.Intersect(InstancedActorBounds))
-			{
-				return Operation(InstanceHandle, InstanceTransform, IterationContext);
-			}
-
-			return true; },
+			return true;
+		},
 		IterationContext, InstancedActorDataPredicate);
 }
 
diff --git a/Engine/Plugins/Runtime/InstancedActors/Source/InstancedActors/Private/InstancedActorsSettingsTypes.cpp b/Engine/Plugins/Runtime/InstancedActors/Source/InstancedActors/Private/InstancedActorsSettingsTypes.cpp
--- a/Engine/Plugins/Runtime/InstancedActors/Source/InstancedActors/Private/InstancedActorsSettingsTypes.cpp
+++ b/Engine/Plugins/Runtime/InstancedActors/Source/InstancedActors/Private/InstancedActorsSettingsTypes.cpp
@@ -133,6 +133,7 @@ void FInstancedActorsSettings::OverrideIfDefault(FConstStructView InOverrideSett
 	IASETTINGS_OVERRIDE_IF_DEFAULT(ActorClass);
 	IASETTINGS_OVERRIDE_IF_DEFAULT(bCanBeDamaged);
 	IASETTINGS_OVERRIDE_IF_DEFAULT(bIgnoreModifierVolumes);
+	IASETTINGS_OVERRIDE_IF_DEFAULT(bModifierVolumeCheckFullyEnclosed);
 	IASETTINGS_OVERRIDE_IF_DEFAULT(bControlPhysicsState);
 	IASETTINGS_OVERRIDE_IF_DEFAULT(GameplayTags);
 	
@@ -159,6 +160,7 @@ FString FInstancedActorsSettings::DebugToString(bool bOverridesOnly) const
 	IASETTINGS_UOBJECT_SETTING_TO_STRING(ActorClass);
 	IASETTINGS_SETTING_TO_STRING(bCanBeDamaged);
 	IASETTINGS_SETTING_TO_STRING(bIgnoreModifierVolumes);
+	IASETTINGS_SETTING_TO_STRING(bModifierVolumeCheckFullyEnclosed);
 	IASETTINGS_SETTING_TO_STRING(bControlPhysicsState);
 	DisplaySettings(SettingsString, GameplayTags, bOverridesOnly, bOverride_GameplayTags, TEXT("GameplayTags"));
 
diff --git a/Engine/Plugins/Runtime/InstancedActors/Source/InstancedActors/Public/InstancedActorsManager.h b/Engine/Plugins/Runtime/InstancedActors/Source/InstancedActors/Public/InstancedActorsManager.h
--- a/Engine/Plugins/Runtime/InstancedActors/Source/InstancedActors/Public/InstancedActorsManager.h
+++ b/Engine/Plugins/Runtime/InstancedActors/Source/InstancedActors/Public/InstancedActorsManager.h
@@ -25,17 +25,29 @@ class UInstancedStaticMeshComponent;
 
 namespace UE::InstancedActors
 {
-	enum EInsideBoundsTestResult
+	enum class EInsideBoundsTestResult
 	{
 		NotInside,
 		OverlapLocation,
 		OverlapBounds
 	};
+	
+	enum class EBoundsTestType
+	{
+		Intersect,
+		Enclosed,
+		
+		Default = Intersect
+	};
 
 	namespace CVars
 	{
 		extern INSTANCEDACTORS_API bool bEnablePersistence;
 	}
+
+	template <typename TBoundsType>
+	bool PassesBoundsTest(const TBoundsType& QueryBounds, EBoundsTestType BoundsTestType, const FInstancedActorsInstanceHandle& InstanceHandle, const FTransform& InstanceTransform);
+
 } // UE::InstancedActors
 
 DECLARE_STATS_GROUP(TEXT("InstanceActor Rendering"), STATGROUP_InstancedActorsRendering, STATCAT_Advanced);
@@ -140,10 +152,6 @@ public:
 
 	/**
 	 * Predicate taking a UInstancedActorsData and returns true if IAD matches search criteria, false otherwise.
-	 * @param InstanceHandle	Handle to the current instance in the iteration
-	 * @param InstanceTransform If entities have been spawned, this will be taken from the Mass transform fragment, else from UInstancedActorsData::InstanceTransforms
-	 * @param InterationContext Provides useful functionality while iterating instances like safe instance deletion
-	 * @return Return true to continue iteration to subsequent instances, false to break iteration.
 	 */
 	using FInstancedActorDataPredicateFunc = TFunctionRef<bool(const UInstancedActorsData& InstancedActorData)>;
 
diff --git a/Engine/Plugins/Runtime/InstancedActors/Source/InstancedActors/Public/InstancedActorsModifiers.h b/Engine/Plugins/Runtime/InstancedActors/Source/InstancedActors/Public/InstancedActorsModifiers.h
--- a/Engine/Plugins/Runtime/InstancedActors/Source/InstancedActors/Public/InstancedActorsModifiers.h
+++ b/Engine/Plugins/Runtime/InstancedActors/Source/InstancedActors/Public/InstancedActorsModifiers.h
@@ -82,15 +82,19 @@ public:
 	template<typename TBoundsType>
 	void ModifyAllInstancesInBounds(const TBoundsType& Bounds, AInstancedActorsManager& Manager, FInstancedActorsIterationContext& IterationContext)
 	{
-		Manager.ForEachInstance(Bounds, [this](const FInstancedActorsInstanceHandle& InstanceHandle, const FTransform& InstanceTransform, FInstancedActorsIterationContext& IterationContext)
+		UE::InstancedActors::EBoundsTestType InstancedActorsDataBoundsTestType{UE::InstancedActors::EBoundsTestType::Default};
+		
+		Manager.ForEachInstance([this, &Bounds, &InstancedActorsDataBoundsTestType](const FInstancedActorsInstanceHandle& InstanceHandle, const FTransform& InstanceTransform, FInstancedActorsIterationContext& IterationContext)
 		{
-			ModifyInstance(InstanceHandle, InstanceTransform, IterationContext);
-	
+			if (UE::InstancedActors::PassesBoundsTest(Bounds, InstancedActorsDataBoundsTestType, InstanceHandle, InstanceTransform))
+			{
+				return ModifyInstance(InstanceHandle, InstanceTransform, IterationContext);
+			}
 			return true;
 		}, 
 		IterationContext,
-		/*Predicate*/TOptional<AInstancedActorsManager::FInstancedActorDataPredicateFunc>([this](const UInstancedActorsData& InstancedActorData)
-		{
+		/*Predicate*/TOptional<AInstancedActorsManager::FInstancedActorDataPredicateFunc>([this, &InstancedActorsDataBoundsTestType](const UInstancedActorsData& InstancedActorData)
+		{			
 			// Allow settings to stop modifiers affect this instance type.
 			const FInstancedActorsSettings* Settings = InstancedActorData.GetSettingsPtr<const FInstancedActorsSettings>();
 			if (Settings && Settings->bOverride_bIgnoreModifierVolumes && Settings->bIgnoreModifierVolumes)
@@ -98,12 +102,17 @@ public:
 				return false;
 			}
 
-			if (!InstanceTagsQuery.IsEmpty())
+			if (InstanceTagsQuery.IsEmpty() || InstancedActorData.GetCombinedTags().MatchesQuery(InstanceTagsQuery))
 			{
-				return InstancedActorData.GetCombinedTags().MatchesQuery(InstanceTagsQuery);
+				// modify the EBoundsTestType for this set of Entities that are associated with this IAD, if the settings require it
+				InstancedActorsDataBoundsTestType = (Settings != nullptr && Settings->bModifierVolumeCheckFullyEnclosed)
+					? UE::InstancedActors::EBoundsTestType::Enclosed
+					: UE::InstancedActors::EBoundsTestType::Intersect;
+
+				return true;
 			}
-			
-			return true;
+
+			return false;
 		}));
 	}
 
diff --git a/Engine/Plugins/Runtime/InstancedActors/Source/InstancedActors/Public/InstancedActorsSettingsTypes.h b/Engine/Plugins/Runtime/InstancedActors/Source/InstancedActors/Public/InstancedActorsSettingsTypes.h
--- a/Engine/Plugins/Runtime/InstancedActors/Source/InstancedActors/Public/InstancedActorsSettingsTypes.h
+++ b/Engine/Plugins/Runtime/InstancedActors/Source/InstancedActors/Public/InstancedActorsSettingsTypes.h
@@ -123,6 +123,9 @@ struct INSTANCEDACTORS_API FInstancedActorsSettings : public FTableRowBase
 	UPROPERTY(EditAnywhere, BlueprintReadWrite, Meta=(InlineEditConditionToggle), Category=InstancedActors)
 	uint8 bOverride_bIgnoreModifierVolumes : 1 = false;
 
+	UPROPERTY(EditAnywhere, BlueprintReadWrite, Meta=(InlineEditConditionToggle), Category=InstancedActors)
+	uint8 bOverride_bModifierVolumeCheckFullyEnclosed : 1 = false;
+
 	UPROPERTY(EditAnywhere, BlueprintReadWrite, Meta=(InlineEditConditionToggle), Category=InstancedActors)
 	uint8 bOverride_bControlPhysicsState : 1 = false;
 
@@ -216,6 +219,10 @@ struct INSTANCEDACTORS_API FInstancedActorsSettings : public FTableRowBase
 	/** Completely disable modifier volumes for this type of instance. */
 	UPROPERTY(EditAnywhere, BlueprintReadWrite, Meta = (EditCondition = "bOverride_bIgnoreModifierVolumes"), Category=InstancedActors)
 	bool bIgnoreModifierVolumes = false;
+
+	/** By default, modifier volumes will do a bounds intersection test against instances. If this is true, modifiers will instead do a fully-enclosed bounds test */
+	UPROPERTY(EditAnywhere, BlueprintReadWrite, Meta = (EditCondition = "bOverride_bModifierVolumeCheckFullyEnclosed"), Category=InstancedActors)
+	bool bModifierVolumeCheckFullyEnclosed = false;
 	
 	/** Gameplay Tags that can be used in Modifier queries */
 	UPROPERTY(EditAnywhere, BlueprintReadWrite, Meta = (EditCondition = "bOverride_GameplayTags"), Category=InstancedActors)
