[InstancedActors] For FInstancedActorsTagSet, don't serialize the hash which represented the combined hashes of tags. This value isn't stable since it's based on hashing FName's integer values, and not its string representation. Modify the == operator to avoid hash collisions (don't simply rely on hash value comparison), and take advantage of the fact that the tag arrays are sorted

diff --git a/Engine/Plugins/Runtime/InstancedActors/Source/InstancedActors/Private/InstancedActorsData.cpp b/Engine/Plugins/Runtime/InstancedActors/Source/InstancedActors/Private/InstancedActorsData.cpp
--- a/Engine/Plugins/Runtime/InstancedActors/Source/InstancedActors/Private/InstancedActorsData.cpp
+++ b/Engine/Plugins/Runtime/InstancedActors/Source/InstancedActors/Private/InstancedActorsData.cpp
@@ -859,6 +859,8 @@ void UInstancedActorsData::PostLoad()
 
 	// Prepare instance delta list for replication and persistence
 	InstanceDeltas.Initialize(*this);
+
+	AdditionalTags.CalcHash();
 }
 
 void UInstancedActorsData::UpdateCullDistance()
diff --git a/Engine/Plugins/Runtime/InstancedActors/Source/InstancedActors/Private/InstancedActorsTypes.cpp b/Engine/Plugins/Runtime/InstancedActors/Source/InstancedActors/Private/InstancedActorsTypes.cpp
--- a/Engine/Plugins/Runtime/InstancedActors/Source/InstancedActors/Private/InstancedActorsTypes.cpp
+++ b/Engine/Plugins/Runtime/InstancedActors/Source/InstancedActors/Private/InstancedActorsTypes.cpp
@@ -68,11 +68,18 @@ FInstancedActorsTagSet::FInstancedActorsTagSet(const FGameplayTagContainer& InTa
 
 	Tags = FGameplayTagContainer::CreateFromArray(SortedTags);
 
-	Hash = 0;
-	for (const FGameplayTag& Tag : SortedTags)
+	CalcHash();
+}
+
+void FInstancedActorsTagSet::CalcHash()
+{
+	uint32 NewHash = 0;
+	for (const FGameplayTag& Tag : Tags.GetGameplayTagArray())
 	{
-		Hash = HashCombine(Hash, GetTypeHash(Tag));
+		NewHash = HashCombine(NewHash, GetTypeHash(Tag));
 	}
+	
+	Hash = NewHash;
 }
 
 //-----------------------------------------------------------------------------
diff --git a/Engine/Plugins/Runtime/InstancedActors/Source/InstancedActors/Public/InstancedActorsTypes.h b/Engine/Plugins/Runtime/InstancedActors/Source/InstancedActors/Public/InstancedActorsTypes.h
--- a/Engine/Plugins/Runtime/InstancedActors/Source/InstancedActors/Public/InstancedActorsTypes.h
+++ b/Engine/Plugins/Runtime/InstancedActors/Source/InstancedActors/Public/InstancedActorsTypes.h
@@ -82,24 +82,40 @@ struct INSTANCEDACTORS_API FInstancedActorsTagSet
 	bool IsEmpty() const { return Tags.IsEmpty(); }
 
 	const FGameplayTagContainer& GetTags() const { return Tags; }
+	
+	void CalcHash();
 
 	bool operator==(const FInstancedActorsTagSet& OtherTagSet) const
 	{
-		return Hash == OtherTagSet.Hash;
-	}
+		if (Hash.IsSet() && OtherTagSet.Hash.IsSet() && Hash.GetValue() != OtherTagSet.Hash.GetValue())
+		{
+			return false;
+		}
 
-	friend uint32 GetTypeHash(const FInstancedActorsTagSet&& TagSet)
-	{
-		return TagSet.Hash;
-	}
+		if (Tags.Num() != OtherTagSet.Tags.Num())
+		{
+			return false;
+		}
 
+		// Tags member will have already been constructed from a sorted TArray
+		for (int32 i = 0; i < Tags.Num(); ++i)
+		{
+			if (Tags.GetGameplayTagArray()[i] != OtherTagSet.Tags.GetGameplayTagArray()[i])
+			{
+				return false;
+			}
+		}
+
+		return true;
+	}
+	
 private:
 
 	UPROPERTY(VisibleAnywhere, Category=InstancedActors)
 	FGameplayTagContainer Tags;
 
-	UPROPERTY()
-	uint32 Hash = 0;
+	UPROPERTY(Transient)
+	TOptional<uint32> Hash;
 };
 
 
