[InstancedActors] Ensure AInstancedActorsManager::CreateNextInstanceActorData() creates a unique child UObject. Iteratively calling NewObject() with the same FName (if one is specified) will give back an existing child UObject if one is found with that FName rather than creating a new one, so we need to create a unique FName each time

diff --git a/Engine/Plugins/Runtime/InstancedActors/Source/InstancedActors/Private/InstancedActorsManager.cpp b/Engine/Plugins/Runtime/InstancedActors/Source/InstancedActors/Private/InstancedActorsManager.cpp
--- a/Engine/Plugins/Runtime/InstancedActors/Source/InstancedActors/Private/InstancedActorsManager.cpp
+++ b/Engine/Plugins/Runtime/InstancedActors/Source/InstancedActors/Private/InstancedActorsManager.cpp
@@ -754,7 +754,8 @@ UInstancedActorsData* AInstancedActorsManager::CreateNextInstanceActorData(TSubc
 	check(InstancedActorsDataClass);
 
 	const FString InstanceDataNameStr = FString::Printf(TEXT("InstancedActorsData_%s"), *ActorClass->GetFName().ToString());
-	UInstancedActorsData* NewInstanceData = NewObject<UInstancedActorsData>(this, InstancedActorsDataClass, FName(InstanceDataNameStr));
+	const FName UniqueName = MakeUniqueObjectName(this, InstancedActorsDataClass, FName(InstanceDataNameStr));
+	UInstancedActorsData* NewInstanceData = NewObject<UInstancedActorsData>(this, InstancedActorsDataClass, UniqueName);
 	// @todo it's conceivable the NextInstanceDataID will overflow. We need to use some handle system in place instead. 
 	NewInstanceData->ID = NextInstanceDataID++;
 	NewInstanceData->ActorClass = ActorClass;
