[InstancedActors] Support for cook-time-created (cell transformer) functionality:
- Adding public AInstancedActorsManager::SetupLoadedInstances() and UInstancedActorsData::SetupLoadedInstances() functions. This logic is called in PostLoad() normally, but if these objects are created at PIE time through cell transformers, they need to be called manually since PostLoad() will not be received in this situation
- GetOrCreateActorInstanceData() given a new bCreateEditorPreviewISMCs parameter, to avoid creating the Editor preview visualization at cook-time

diff --git a/Engine/Plugins/Runtime/InstancedActors/Source/InstancedActors/Private/InstancedActorsData.cpp b/Engine/Plugins/Runtime/InstancedActors/Source/InstancedActors/Private/InstancedActorsData.cpp
--- a/Engine/Plugins/Runtime/InstancedActors/Source/InstancedActors/Private/InstancedActorsData.cpp
+++ b/Engine/Plugins/Runtime/InstancedActors/Source/InstancedActors/Private/InstancedActorsData.cpp
@@ -806,6 +806,18 @@ void UInstancedActorsData::PostLoad()
 {
 	Super::PostLoad();
 
+	SetupLoadedInstances();
+}
+
+void UInstancedActorsData::SetupLoadedInstances()
+{
+	if (bHasSetupLoadedInstances)
+	{
+		return;
+	}
+	
+	bHasSetupLoadedInstances = true;
+	
 	NumInstances = InstanceTransforms.Num();
 
 	// Cache asset bounds
diff --git a/Engine/Plugins/Runtime/InstancedActors/Source/InstancedActors/Private/InstancedActorsManager.cpp b/Engine/Plugins/Runtime/InstancedActors/Source/InstancedActors/Private/InstancedActorsManager.cpp
--- a/Engine/Plugins/Runtime/InstancedActors/Source/InstancedActors/Private/InstancedActorsManager.cpp
+++ b/Engine/Plugins/Runtime/InstancedActors/Source/InstancedActors/Private/InstancedActorsManager.cpp
@@ -431,6 +431,18 @@ void AInstancedActorsManager::PostLoad()
 {
 	Super::PostLoad();
 
+	SetupLoadedInstances();
+}
+
+void AInstancedActorsManager::SetupLoadedInstances()
+{
+	if (bHasSetupLoadedInstances)
+	{
+		return;
+	}
+	
+	bHasSetupLoadedInstances = true;
+
 	// @todo need default implementation of GUID and set it here via SetSavedActorGUID
 
 	for (UInstancedActorsData* InstanceData : PerActorClassInstanceData)
@@ -768,7 +780,7 @@ UInstancedActorsData* AInstancedActorsManager::CreateNextInstanceActorData(TSubc
 	return NewInstanceData;
 }
 
-UInstancedActorsData& AInstancedActorsManager::GetOrCreateActorInstanceData(TSubclassOf<AActor> ActorClass, const FInstancedActorsTagSet& AdditionalInstanceTags)
+UInstancedActorsData& AInstancedActorsManager::GetOrCreateActorInstanceData(TSubclassOf<AActor> ActorClass, const FInstancedActorsTagSet& AdditionalInstanceTags, bool bCreateEditorPreviewISMCs)
 {
 	checkf(HasActorBegunPlay() == false, TEXT("AInstancedActorsManager doesn't yet support runtime addition of instances"));
 
@@ -786,29 +798,32 @@ UInstancedActorsData& AInstancedActorsManager::GetOrCreateActorInstanceData(TSub
 	UInstancedActorsData* NewInstanceData = CreateNextInstanceActorData(ActorClass, AdditionalInstanceTags);
 	PerActorClassInstanceData.Add(NewInstanceData);
 
-	// Get or create exemplar actor to derive entities from
-	UWorld* World = GetWorld();
-	check(World);
-	UInstancedActorsSubsystem* EditorInstancedActorSubsystem = UE::InstancedActors::Utils::GetInstancedActorsSubsystem(*World);
-	check(EditorInstancedActorSubsystem);
-	// @todo For performance reasons it would be helpful to be able to cache this into UInstancedActorsData::ExemplarActorData
-	// This would mean that undo-ing a conversion of actors to IAM in Editor will not clear or destroy the exemplar actor/map entry,
-	// which is the reason for the sanity-check of the actor ptr in UInstancedActorsSubsystem::GetOrCreateExemplarActor().
-	// We need to also investigate why converting the IAM to actors (destroying the IAM) does not seem to trigger GC on the UInstancedActorsData
-	TSharedPtr<UE::InstancedActors::FExemplarActorData> ExemplarActorData = EditorInstancedActorSubsystem->GetOrCreateExemplarActor(ActorClass);
-	const AActor* const ExemplarActor = ExemplarActorData->Actor.Get();
-	check(ExemplarActor);
-
-	// Create stand-in 'editor only' ISMC's to preview instances in the level editor. These will be stripped
-	// out at cook time and in PostLoad for game worlds (PIE)
-	FInstancedActorsVisualizationDesc DefaultVisualiation = EditorInstancedActorSubsystem->CreateVisualDescriptionFromActor(*ExemplarActor);
-	CreateISMComponents(DefaultVisualiation, NewInstanceData->SharedSettings, NewInstanceData->EditorPreviewISMComponents, /*bEditorPreviewISMCs*/ true);
-
-	// Add EditorPreviewISMComponents to the known ISMComponent-to-InstanceData map (ISMComponentToInstanceDataMap)
-	// so that we can identify Instanced Actors instances by ISM instance id coming from one of the EditorPreviewISMComponents
-	// (like when the instance is clicked in the editor). 
-	// @see ActorInstanceHandleFromFSMInstanceId 
-	RegisterInstanceDatasComponents(*NewInstanceData, NewInstanceData->EditorPreviewISMComponents);
+	if (bCreateEditorPreviewISMCs)
+	{
+		// Get or create exemplar actor to derive entities from
+		UWorld* World = GetWorld();
+		check(World);
+		UInstancedActorsSubsystem* EditorInstancedActorSubsystem = UE::InstancedActors::Utils::GetInstancedActorsSubsystem(*World);
+		check(EditorInstancedActorSubsystem);
+		// @todo For performance reasons it would be helpful to be able to cache this into UInstancedActorsData::ExemplarActorData
+		// This would mean that undo-ing a conversion of actors to IAM in Editor will not clear or destroy the exemplar actor/map entry,
+		// which is the reason for the sanity-check of the actor ptr in UInstancedActorsSubsystem::GetOrCreateExemplarActor().
+		// We need to also investigate why converting the IAM to actors (destroying the IAM) does not seem to trigger GC on the UInstancedActorsData
+		TSharedPtr<UE::InstancedActors::FExemplarActorData> ExemplarActorData = EditorInstancedActorSubsystem->GetOrCreateExemplarActor(ActorClass);
+		const AActor* const ExemplarActor = ExemplarActorData->Actor.Get();
+		check(ExemplarActor);
+
+		// Create stand-in 'editor only' ISMC's to preview instances in the level editor. These will be stripped
+		// out at cook time and in PostLoad for game worlds (PIE)
+		FInstancedActorsVisualizationDesc DefaultVisualiation = EditorInstancedActorSubsystem->CreateVisualDescriptionFromActor(*ExemplarActor);
+		CreateISMComponents(DefaultVisualiation, NewInstanceData->SharedSettings, NewInstanceData->EditorPreviewISMComponents, /*bEditorPreviewISMCs*/ true);
+
+		// Add EditorPreviewISMComponents to the known ISMComponent-to-InstanceData map (ISMComponentToInstanceDataMap)
+		// so that we can identify Instanced Actors instances by ISM instance id coming from one of the EditorPreviewISMComponents
+		// (like when the instance is clicked in the editor). 
+		// @see ActorInstanceHandleFromFSMInstanceId 
+		RegisterInstanceDatasComponents(*NewInstanceData, NewInstanceData->EditorPreviewISMComponents);		
+	}
 
 	// Cache asset bounds for use during instance population
 	NewInstanceData->CachedLocalBounds = CalculateBounds(ActorClass);
diff --git a/Engine/Plugins/Runtime/InstancedActors/Source/InstancedActors/Public/InstancedActorsData.h b/Engine/Plugins/Runtime/InstancedActors/Source/InstancedActors/Public/InstancedActorsData.h
--- a/Engine/Plugins/Runtime/InstancedActors/Source/InstancedActors/Public/InstancedActorsData.h
+++ b/Engine/Plugins/Runtime/InstancedActors/Source/InstancedActors/Public/InstancedActorsData.h
@@ -197,6 +197,10 @@ public:
 	// Returns true if InstanceHandle refers to this instance data and we have current information for an
 	// instance at InstanceHandle.InstanceIndex
 	bool IsValidInstance(const FInstancedActorsInstanceHandle& InstanceHandle) const;
+	
+	// Performs setup after all Instances have been loaded. Canonically called from PostLoad(), but may need to be called manually
+	// if this UInstancedActorsData is created at cook/runtime
+	void SetupLoadedInstances();
 
 #if WITH_EDITOR
 	// Add's an instance of this actor type at Transform
@@ -445,6 +449,9 @@ protected:
 	// @return The unlinked actor (if one was currently linked, nullptr otherwise)
 	AActor* UnlinkActor(FInstancedActorsInstanceIndex InstanceToUnlink);
 
+	// True if SetupLoadedInstances has ever been called
+	bool bHasSetupLoadedInstances : 1 = false;
+	
 	// True if Initialize has ever been called
 	bool bHasEverInitialized : 1 = false;
 
diff --git a/Engine/Plugins/Runtime/InstancedActors/Source/InstancedActors/Public/InstancedActorsManager.h b/Engine/Plugins/Runtime/InstancedActors/Source/InstancedActors/Public/InstancedActorsManager.h
--- a/Engine/Plugins/Runtime/InstancedActors/Source/InstancedActors/Public/InstancedActorsManager.h
+++ b/Engine/Plugins/Runtime/InstancedActors/Source/InstancedActors/Public/InstancedActorsManager.h
@@ -82,6 +82,12 @@ public:
 
 	FInstancedActorsManagerHandle GetManagerHandle() const;
 
+	/** 
+	 * Performs setup after all Instances have been loaded. Canonically called from PostLoad(), but may need to be called manually
+	 * if this AInstancedActorsManager is created at cook/runtime
+	 */
+	void SetupLoadedInstances();
+
 	/** 
 	 * Initializes all PerActorClassInstanceData, applies pre-spawn modifiers, spawns entities then applies post-spawn modifiers.
 	 * 
@@ -299,9 +305,11 @@ protected:
 	FBox CalculateLocalInstanceBounds() const;
 
 #if WITH_EDITOR
+public:
 	/** Helper function to create and initialize per-actor-class UInstancedActorsData's, optionally further partitioned by AdditionalInstanceTags */
-	UInstancedActorsData& GetOrCreateActorInstanceData(TSubclassOf<AActor> ActorClass, const FInstancedActorsTagSet& AdditionalInstanceTags);
+	UInstancedActorsData& GetOrCreateActorInstanceData(TSubclassOf<AActor> ActorClass, const FInstancedActorsTagSet& AdditionalInstanceTags, bool bCreateEditorPreviewISMCs = true);
 
+protected:
 	virtual UInstancedActorsData* CreateNextInstanceActorData(TSubclassOf<AActor> ActorClass, const FInstancedActorsTagSet& AdditionalInstanceTags);
 
 	/** Used to set the right properties on the editor ISMCs so we can do per-instance selection. */
@@ -327,6 +335,9 @@ protected:
 	UPROPERTY(Transient)
 	bool bHasSpawnedEntities = false;
 
+	/** True if SetupLoadedInstances has ever been called */
+	bool bHasSetupLoadedInstances : 1 = false;
+
 	/** 
 	 * Incremented in GetOrCreateActorInstanceData to provide IAD's with a stable, unique identifier
 	 * within this IAM.
