[InstancedActors] Added an ensureAlwaysMsgf() before the call into SetUpEditorPreviewISMComponent() which uses a TNotNull parameter, for better logging to trap if an element of EditorPreviewISMComponents is null

diff --git a/Engine/Plugins/Runtime/InstancedActors/Source/InstancedActors/Private/InstancedActorsManager.cpp b/Engine/Plugins/Runtime/InstancedActors/Source/InstancedActors/Private/InstancedActorsManager.cpp
--- a/Engine/Plugins/Runtime/InstancedActors/Source/InstancedActors/Private/InstancedActorsManager.cpp
+++ b/Engine/Plugins/Runtime/InstancedActors/Source/InstancedActors/Private/InstancedActorsManager.cpp
@@ -841,6 +841,12 @@ void AInstancedActorsManager::PreRegisterAllComponents()
 		// Modify the ISMCs here since the components might have been serialized before this change.
 		for (TObjectPtr<UInstancedStaticMeshComponent> ISMComponent : InstanceData->EditorPreviewISMComponents)
 		{
+			if (ISMComponent == nullptr)
+			{
+				ensureAlwaysMsgf(false, TEXT("UInstancedActorsData has a null EditorPreviewISMComponents element: %s"), *InstanceData->GetDebugName(/*bCompact*/ false));
+				continue;
+			}
+
 			SetUpEditorPreviewISMComponent(ISMComponent);
 		}
 
