[InstancedActors] Fix a bug by only applying FISMComponentDescriptor::LocalTransform to visible UStaticMeshComponents, and those that are not the Actor's RootComponent. This will allow the instanced actor transform to apply to the root correctly, and addresses issues with setups that involve collision-only static meshes

diff --git a/Engine/Plugins/Runtime/InstancedActors/Source/InstancedActors/Private/InstancedActorsTypes.cpp b/Engine/Plugins/Runtime/InstancedActors/Source/InstancedActors/Private/InstancedActorsTypes.cpp
--- a/Engine/Plugins/Runtime/InstancedActors/Source/InstancedActors/Private/InstancedActorsTypes.cpp
+++ b/Engine/Plugins/Runtime/InstancedActors/Source/InstancedActors/Private/InstancedActorsTypes.cpp
@@ -110,9 +110,15 @@ FInstancedActorsVisualizationDesc FInstancedActorsVisualizationDesc::FromActor(c
 	const FVisualizationDescSetupFunction& AdditionalSetupSteps)
 {
 	FInstancedActorsVisualizationDesc Visualization;
+	USceneComponent* RootComponent = ExemplarActor.GetRootComponent();
 	
-	ExemplarActor.ForEachComponent<UStaticMeshComponent>(/*bIncludeFromChildActors=*/false, [&Visualization](UStaticMeshComponent* SourceStaticMeshComponent)
+	ExemplarActor.ForEachComponent<UStaticMeshComponent>(/*bIncludeFromChildActors=*/false, [&Visualization, RootComponent](UStaticMeshComponent* SourceStaticMeshComponent)
 	{
+		if (!SourceStaticMeshComponent->IsVisible())
+		{
+			return;
+		}
+
 		const UStaticMesh* StaticMesh = SourceStaticMeshComponent->GetStaticMesh();
 		if (!IsValid(StaticMesh))
 		{
@@ -125,7 +131,10 @@ FInstancedActorsVisualizationDesc FInstancedActorsVisualizationDesc::FromActor(c
 
 		// LocalTransform means local to the Actor/Entity, so we need to compute based on the UStaticMeshComponent's relative transform accordingly
 		// (in case this UStaticMeshComponent was a child of another UStaticMeshComponent within the Actor hierarchy)
-		ISMComponentDescriptor.LocalTransform = SourceStaticMeshComponent->GetComponentToWorld();
+		if (SourceStaticMeshComponent != RootComponent)
+		{
+			ISMComponentDescriptor.LocalTransform = SourceStaticMeshComponent->GetComponentToWorld();
+		}
 	});
 
 	AdditionalSetupSteps(ExemplarActor, Visualization);
