[InstancedActors] Temporary fix to address timing issues in which UInstancedActorsStationaryLODBatchProcessor::ConfigureQueries() gets called before UGameFeatureAction_ConfigureInstancedActors has activated and thus set the subsystem requirements of the desired InstancedActorsSubsystemClass. Made UInstancedActorsStationaryLODBatchProcessor::Execute() ignore the setup requirement, and just grab the class as defined by the UGameFeatureAction_ConfigureInstancedActors once it has activated.

diff --git a/Engine/Plugins/Runtime/InstancedActors/Source/InstancedActors/Private/InstancedActorsStationaryLODBatchProcessor.cpp b/Engine/Plugins/Runtime/InstancedActors/Source/InstancedActors/Private/InstancedActorsStationaryLODBatchProcessor.cpp
--- a/Engine/Plugins/Runtime/InstancedActors/Source/InstancedActors/Private/InstancedActorsStationaryLODBatchProcessor.cpp
+++ b/Engine/Plugins/Runtime/InstancedActors/Source/InstancedActors/Private/InstancedActorsStationaryLODBatchProcessor.cpp
@@ -177,7 +177,9 @@ void UInstancedActorsStationaryLODBatchProcessor::Execute(FMassEntityManager& En
 	{
 		return;
 	}
-	UInstancedActorsSubsystem* InstancedActorSubsystem = Context.GetMutableSubsystem<UInstancedActorsSubsystem>(InstancedActorsSubsystemClass);
+	//@todo: revisit this temporary fix. This was done because UGameFeatureAction_ConfigureInstancedActors::OnGameFeatureActivating() may happen after UInstancedActorsStationaryLODBatchProcessor::ConfigureQueries()
+	//UInstancedActorsSubsystem* InstancedActorSubsystem = Context.GetMutableSubsystem<UInstancedActorsSubsystem>(InstancedActorsSubsystemClass);
+	UInstancedActorsSubsystem* InstancedActorSubsystem = UE::InstancedActors::Utils::GetInstancedActorsSubsystem(*Context.GetWorld());
 	if (!ensureMsgf(InstancedActorSubsystem, TEXT("UInstancedActorsSubsystem is missing, this is unexpected")))
 	{
 		return;
