Change FInstancedActorsInstanceHandle to store a weak reference to the UInstancedActorData instead of strong, to prevent garbage references if the level with the UInstancedActorData is streamed out.

diff --git a/Engine/Plugins/Runtime/InstancedActors/Source/InstancedActors/Private/InstancedActorsIndex.cpp b/Engine/Plugins/Runtime/InstancedActors/Source/InstancedActors/Private/InstancedActorsIndex.cpp
--- a/Engine/Plugins/Runtime/InstancedActors/Source/InstancedActors/Private/InstancedActorsIndex.cpp
+++ b/Engine/Plugins/Runtime/InstancedActors/Source/InstancedActors/Private/InstancedActorsIndex.cpp
@@ -18,20 +18,21 @@ FString FInstancedActorsInstanceIndex::GetDebugName() const
 
 bool FInstancedActorsInstanceHandle::IsValid() const 
 {
-	return ::IsValid(InstancedActorData) && Index.IsValid();
+	return InstancedActorData.IsValid() && Index.IsValid();
 }
 
 UInstancedActorsData& FInstancedActorsInstanceHandle::GetInstanceActorDataChecked() const 
 { 
-	check(::IsValid(InstancedActorData));
+	UInstancedActorsData* RawInstancedActorData = InstancedActorData.Get();
+	check(RawInstancedActorData);
 	return *InstancedActorData; 
 }
 
 AInstancedActorsManager* FInstancedActorsInstanceHandle::GetManager() const
 {
-	if (::IsValid(InstancedActorData))
+	if (UInstancedActorsData* RawInstancedActorData = InstancedActorData.Get())
 	{
-		return InstancedActorData->GetManager();
+		return RawInstancedActorData->GetManager();
 	}
 
 	return nullptr;
@@ -44,15 +45,16 @@ AInstancedActorsManager& FInstancedActorsInstanceHandle::GetManagerChecked() con
 
 FString FInstancedActorsInstanceHandle::GetDebugName() const
 {
-	return FString::Printf(TEXT("%s : %s"), InstancedActorData ? *InstancedActorData->GetDebugName() : TEXT("null"), *Index.GetDebugName());
+	UInstancedActorsData* RawInstancedActorData = InstancedActorData.Get();
+	return FString::Printf(TEXT("%s : %s"), RawInstancedActorData ? *RawInstancedActorData->GetDebugName() : TEXT("null"), *Index.GetDebugName());
 }
 
 uint32 GetTypeHash(const FInstancedActorsInstanceHandle& Handle)
 {
 	uint32 Hash = 0;
-	if (::IsValid(Handle.InstancedActorData))
+	if (UInstancedActorsData* RawInstancedActorData = Handle.InstancedActorData.Get())
 	{
-		Hash = GetTypeHash(Handle.InstancedActorData);
+		Hash = GetTypeHash(RawInstancedActorData);
 	}
 	Hash = HashCombine(Hash, Handle.GetIndex());
 
diff --git a/Engine/Plugins/Runtime/InstancedActors/Source/InstancedActors/Private/InstancedActorsIteration.cpp b/Engine/Plugins/Runtime/InstancedActors/Source/InstancedActors/Private/InstancedActorsIteration.cpp
--- a/Engine/Plugins/Runtime/InstancedActors/Source/InstancedActors/Private/InstancedActorsIteration.cpp
+++ b/Engine/Plugins/Runtime/InstancedActors/Source/InstancedActors/Private/InstancedActorsIteration.cpp
@@ -16,7 +16,7 @@ void FInstancedActorsIterationContext::RemoveInstanceDeferred(const FInstancedAc
 {
 	if (ensure(InstanceHandle.IsValid()))
 	{
-		InstancesToRemove.FindOrAdd(InstanceHandle.InstancedActorData).Add(InstanceHandle.GetInstanceIndex());
+		InstancesToRemove.FindOrAdd(InstanceHandle.GetInstanceActorData()).Add(InstanceHandle.GetInstanceIndex());
 	}
 }
 
diff --git a/Engine/Plugins/Runtime/InstancedActors/Source/InstancedActors/Public/InstancedActorsIndex.h b/Engine/Plugins/Runtime/InstancedActors/Source/InstancedActors/Public/InstancedActorsIndex.h
--- a/Engine/Plugins/Runtime/InstancedActors/Source/InstancedActors/Public/InstancedActorsIndex.h
+++ b/Engine/Plugins/Runtime/InstancedActors/Source/InstancedActors/Public/InstancedActorsIndex.h
@@ -94,7 +94,7 @@ struct INSTANCEDACTORS_API FInstancedActorsInstanceHandle
 	FInstancedActorsInstanceHandle() = default;
 	FInstancedActorsInstanceHandle(UInstancedActorsData& InInstancedActorData, FInstancedActorsInstanceIndex InIndex);
 
-	UInstancedActorsData* GetInstanceActorData() const { return InstancedActorData; }
+	UInstancedActorsData* GetInstanceActorData() const { return InstancedActorData.Get(); }
 	UInstancedActorsData& GetInstanceActorDataChecked() const;
 
 	AInstancedActorsManager* GetManager() const;
@@ -124,7 +124,7 @@ private:
 
 	/** Specific UInstancedActorsData responsible for this instance.Can be used to get to it's owning Manager. */
 	UPROPERTY()
-	TObjectPtr<UInstancedActorsData> InstancedActorData = nullptr;
+	TWeakObjectPtr<UInstancedActorsData> InstancedActorData = nullptr;
 
 	/** Stable(consistent between client and server) instance index into UInstancedActorsData. */
 	UPROPERTY()
