PR #12788: InstancedActors: Project Settings Fixes

Making sure UInstancedActorsProjectSettings.CompiledActiveConfig is up to date after engine loading as well as getting udpated every time a change to settings is done.
The CL takes advantage of that by ensuring that the apropriate InstancedActorsSubsystem class is used when converting actors to InstancedActors via the editor's Actor menu option

diff --git a/Engine/Plugins/Runtime/InstancedActors/Source/InstancedActors/Private/InstancedActorsSettings.cpp b/Engine/Plugins/Runtime/InstancedActors/Source/InstancedActors/Private/InstancedActorsSettings.cpp
--- a/Engine/Plugins/Runtime/InstancedActors/Source/InstancedActors/Private/InstancedActorsSettings.cpp
+++ b/Engine/Plugins/Runtime/InstancedActors/Source/InstancedActors/Private/InstancedActorsSettings.cpp
@@ -20,6 +20,28 @@ UInstancedActorsProjectSettings::UInstancedActorsProjectSettings()
 	CompiledActiveConfig = DefaultConfig;
 }
 
+void UInstancedActorsProjectSettings::PostInitProperties()
+{
+	Super::PostInitProperties();
+	CompiledActiveConfig = DefaultConfig;
+}
+
+#if WITH_EDITOR
+void UInstancedActorsProjectSettings::PostEditChangeProperty(struct FPropertyChangedEvent& PropertyChangedEvent)
+{
+	if (ClassConfigOverrides.IsEmpty())
+	{
+		CompiledActiveConfig = DefaultConfig;
+	}
+	else
+	{
+		CompileSettings();
+	}
+
+	Super::PostEditChangeProperty(PropertyChangedEvent);
+}
+#endif
+
 TSubclassOf<UMassActorSpawnerSubsystem> UInstancedActorsProjectSettings::GetServerActorSpawnerSubsystemClass() const 
 { 
 	return CompiledActiveConfig.ServerActorSpawnerSubsystemClass;
diff --git a/Engine/Plugins/Runtime/InstancedActors/Source/InstancedActors/Public/InstancedActorsSettings.h b/Engine/Plugins/Runtime/InstancedActors/Source/InstancedActors/Public/InstancedActorsSettings.h
--- a/Engine/Plugins/Runtime/InstancedActors/Source/InstancedActors/Public/InstancedActorsSettings.h
+++ b/Engine/Plugins/Runtime/InstancedActors/Source/InstancedActors/Public/InstancedActorsSettings.h
@@ -47,7 +47,7 @@ struct FClassConfigOverrideEntry
  * @see FInstancedActorsClassSettingsBase and FInstancedActorsClassSettings for per-class specific runtime settings.
  * @see AInstancedActorsManager
  */
-UCLASS(Config=InstancedActors, defaultconfig, DisplayName = "Instanced Actors")
+UCLASS(Config=InstancedActors, defaultconfig, DisplayName = "Instanced Actors", MinimalAPI)
 class UInstancedActorsProjectSettings : public UDeveloperSettings
 {
 	GENERATED_BODY()
@@ -57,10 +57,10 @@ public:
 
 	UInstancedActorsProjectSettings();
 
-	TSubclassOf<UMassActorSpawnerSubsystem> GetServerActorSpawnerSubsystemClass() const;
-	TSubclassOf<UMassActorSpawnerSubsystem> GetClientActorSpawnerSubsystemClass() const;
-	TSubclassOf<UInstancedActorsSubsystem> GetInstancedActorsSubsystemClass() const;
-	TSubclassOf<UMassStationaryDistanceVisualizationTrait> GetStationaryVisualizationTraitClass() const;
+	INSTANCEDACTORS_API TSubclassOf<UMassActorSpawnerSubsystem> GetServerActorSpawnerSubsystemClass() const;
+	INSTANCEDACTORS_API TSubclassOf<UMassActorSpawnerSubsystem> GetClientActorSpawnerSubsystemClass() const;
+	INSTANCEDACTORS_API TSubclassOf<UInstancedActorsSubsystem> GetInstancedActorsSubsystemClass() const;
+	INSTANCEDACTORS_API TSubclassOf<UMassStationaryDistanceVisualizationTrait> GetStationaryVisualizationTraitClass() const;
 
 	void RegisterConfigOverride(UObject& Owner, const FInstancedActorsConfig& Config);
 	void UnregisterConfigOverride(UObject& Owner);
@@ -68,6 +68,11 @@ public:
 	FOnSettingsChanged& GetOnSettingsUpdated() { return OnSettingsUpdated; }
 
 protected:
+	virtual void PostInitProperties() override;
+#if WITH_EDITOR
+	virtual void PostEditChangeProperty(struct FPropertyChangedEvent& PropertyChangedEvent) override;
+#endif
+
 	void CompileSettings();
 
 public:
diff --git a/Engine/Plugins/Runtime/InstancedActors/Source/InstancedActorsEditor/Private/InstancedActorsEditorModule.cpp b/Engine/Plugins/Runtime/InstancedActors/Source/InstancedActorsEditor/Private/InstancedActorsEditorModule.cpp
--- a/Engine/Plugins/Runtime/InstancedActors/Source/InstancedActorsEditor/Private/InstancedActorsEditorModule.cpp
+++ b/Engine/Plugins/Runtime/InstancedActors/Source/InstancedActorsEditor/Private/InstancedActorsEditorModule.cpp
@@ -5,6 +5,7 @@
 #include "InstancedActorsManager.h"
 #include "InstancedActorsData.h"
 #include "InstancedActorsIteration.h"
+#include "InstancedActorsSettings.h"
 #include "Modules/ModuleManager.h"
 #include "LevelEditor.h"
 #include "Engine/World.h"
@@ -125,7 +126,7 @@ TSharedRef<FExtender> FInstancedActorsEditorModule::CreateLevelViewportContextMe
 
 void FInstancedActorsEditorModule::ConvertActorsToIAsUIAction(TConstArrayView<AActor*> InActors) const
 {
-	CustomizedConvertActorsToIAsUIAction(InActors, UInstancedActorsSubsystem::StaticClass());
+	CustomizedConvertActorsToIAsUIAction(InActors, GET_INSTANCEDACTORS_CONFIG_VALUE(GetInstancedActorsSubsystemClass()));
 }
 
 void FInstancedActorsEditorModule::CustomizedConvertActorsToIAsUIAction(TConstArrayView<AActor*> InActors
