Added ModifyVisualDescriptionForActor to InstancedActorSubsystem allowing to update the visual after initialization

diff --git a/Engine/Plugins/Runtime/InstancedActors/Source/InstancedActors/Private/InstancedActorsData.cpp b/Engine/Plugins/Runtime/InstancedActors/Source/InstancedActors/Private/InstancedActorsData.cpp
--- a/Engine/Plugins/Runtime/InstancedActors/Source/InstancedActors/Private/InstancedActorsData.cpp
+++ b/Engine/Plugins/Runtime/InstancedActors/Source/InstancedActors/Private/InstancedActorsData.cpp
@@ -1524,6 +1524,10 @@ uint8 UInstancedActorsData::AddVisualizationAsync(const FInstancedActorsSoftVisu
 	{
 		// Nothing needing to load. Resolve hard visualization description and add immediately.
 		FInstancedActorsVisualizationDesc VisualizationDesc(SoftVisualizationDesc);
+		if (ExemplarActorData.IsValid() && ExemplarActorData->Actor)
+		{
+			GetManagerChecked().GetInstancedActorSubsystemChecked().ModifyVisualDescriptionForActor(TNotNull<AActor*>(ExemplarActorData->Actor.Get()), VisualizationDesc);
+		}
 		const uint8 NewVisualizationIndex = AddVisualization(VisualizationDesc);
 		return NewVisualizationIndex;
 	}
@@ -1536,8 +1540,10 @@ uint8 UInstancedActorsData::AddVisualizationAsync(const FInstancedActorsSoftVisu
 		check(InstanceVisualizations.IsValidIndex(ReservedVisualizationIndex));
 		FInstancedActorsVisualizationInfo& ReservedVisualization = InstanceVisualizations[ReservedVisualizationIndex];
 
+		TWeakObjectPtr<UInstancedActorsSubsystem> WeakInstancedActorsSubsystem = GetManagerChecked().GetInstancedActorSubsystem();
+
 		// Async load assets and complete initialization when load completes
-		ReservedVisualization.AssetLoadHandle = UAssetManager::Get().GetStreamableManager().RequestAsyncLoad(VisualizationAssetsToLoad, FStreamableDelegate::CreateWeakLambda(this, [this, SoftVisualizationDesc, ReservedVisualizationIndex]()
+		ReservedVisualization.AssetLoadHandle = UAssetManager::Get().GetStreamableManager().RequestAsyncLoad(VisualizationAssetsToLoad, FStreamableDelegate::CreateWeakLambda(this, [this, SoftVisualizationDesc, ReservedVisualizationIndex, WeakInstancedActorsSubsystem]()
 			{
 				// Check visualization is still allocated
 				if (!ensure(InstanceVisualizations.IsValidIndex(ReservedVisualizationIndex)))
@@ -1554,6 +1560,11 @@ uint8 UInstancedActorsData::AddVisualizationAsync(const FInstancedActorsSoftVisu
 				// Resolve hard visualization description
 				FInstancedActorsVisualizationDesc VisualizationDesc(SoftVisualizationDesc);
 
+				if (WeakInstancedActorsSubsystem.IsValid() && ExemplarActorData.IsValid() && ExemplarActorData->Actor)
+				{
+					WeakInstancedActorsSubsystem->ModifyVisualDescriptionForActor(TNotNull<AActor*>(ExemplarActorData->Actor.Get()), VisualizationDesc);
+				}
+
 				// Init reserved visualization
 				InitializeVisualization(ReservedVisualizationIndex, VisualizationDesc);
 
diff --git a/Engine/Plugins/Runtime/InstancedActors/Source/InstancedActors/Public/InstancedActorsSubsystem.h b/Engine/Plugins/Runtime/InstancedActors/Source/InstancedActors/Public/InstancedActorsSubsystem.h
--- a/Engine/Plugins/Runtime/InstancedActors/Source/InstancedActors/Public/InstancedActorsSubsystem.h
+++ b/Engine/Plugins/Runtime/InstancedActors/Source/InstancedActors/Public/InstancedActorsSubsystem.h
@@ -174,6 +174,10 @@ public:
 	INSTANCEDACTORS_API void PopAllDirtyRepresentationInstances(TArray<FInstancedActorsInstanceHandle>& OutInstances);
 
 	INSTANCEDACTORS_API virtual FInstancedActorsVisualizationDesc CreateVisualDescriptionFromActor(const AActor& ExemplarActor) const;
+	/*
+	* Called when an additional/alternate VisualizationDesc is registered. Override to make custom modifications to the visual representation
+	*/
+	INSTANCEDACTORS_API virtual void ModifyVisualDescriptionForActor(const TNotNull<AActor*> ExemplarActor, FInstancedActorsVisualizationDesc& InOutVisualization) const {}
 
 	//~ Begin UTickableWorldSubsystem Overrides
 	INSTANCEDACTORS_API virtual void Tick(float DeltaTime) override;
