Moved the ModifyVisualDescriptionForActor call to AddVisualization

diff --git a/Engine/Plugins/Runtime/InstancedActors/Source/InstancedActors/Private/InstancedActorsData.cpp b/Engine/Plugins/Runtime/InstancedActors/Source/InstancedActors/Private/InstancedActorsData.cpp
--- a/Engine/Plugins/Runtime/InstancedActors/Source/InstancedActors/Private/InstancedActorsData.cpp
+++ b/Engine/Plugins/Runtime/InstancedActors/Source/InstancedActors/Private/InstancedActorsData.cpp
@@ -1502,13 +1502,18 @@ void UInstancedActorsData::InitializeVisualization(uint8 AllocatedVisualizationI
 	}
 }
 
-uint8 UInstancedActorsData::AddVisualization(const FInstancedActorsVisualizationDesc& VisualizationDesc)
+uint8 UInstancedActorsData::AddVisualization(FInstancedActorsVisualizationDesc& InOutVisualizationDesc)
 {
 	// Reuse free or create new InstanceVisualizations entry
 	uint8 NewVisualizationIndex = AllocateVisualization();
 
+	if (ExemplarActorData.IsValid() && ExemplarActorData->Actor)
+	{
+		GetManagerChecked().GetInstancedActorSubsystemChecked().ModifyVisualDescriptionForActor(TNotNull<AActor*>(ExemplarActorData->Actor.Get()), InOutVisualizationDesc);
+	}
+
 	// Init new visualization
-	InitializeVisualization(NewVisualizationIndex, VisualizationDesc);
+	InitializeVisualization(NewVisualizationIndex, InOutVisualizationDesc);
 
 	return NewVisualizationIndex;
 }
@@ -1524,10 +1529,6 @@ uint8 UInstancedActorsData::AddVisualizationAsync(const FInstancedActorsSoftVisu
 	{
 		// Nothing needing to load. Resolve hard visualization description and add immediately.
 		FInstancedActorsVisualizationDesc VisualizationDesc(SoftVisualizationDesc);
-		if (ExemplarActorData.IsValid() && ExemplarActorData->Actor)
-		{
-			GetManagerChecked().GetInstancedActorSubsystemChecked().ModifyVisualDescriptionForActor(TNotNull<AActor*>(ExemplarActorData->Actor.Get()), VisualizationDesc);
-		}
 		const uint8 NewVisualizationIndex = AddVisualization(VisualizationDesc);
 		return NewVisualizationIndex;
 	}
diff --git a/Engine/Plugins/Runtime/InstancedActors/Source/InstancedActors/Public/InstancedActorsData.h b/Engine/Plugins/Runtime/InstancedActors/Source/InstancedActors/Public/InstancedActorsData.h
--- a/Engine/Plugins/Runtime/InstancedActors/Source/InstancedActors/Public/InstancedActorsData.h
+++ b/Engine/Plugins/Runtime/InstancedActors/Source/InstancedActors/Public/InstancedActorsData.h
@@ -308,7 +308,7 @@ public:
 	// for each VisualizationDesc.InstancedMeshes
 	// @warning No more than 254 visualizations may be registered at any time to allow for uint8 indexing.
 	// @return Index handle to refer to the registered visualization
-	UE_API uint8 AddVisualization(const FInstancedActorsVisualizationDesc& VisualizationDesc);
+	UE_API uint8 AddVisualization(FInstancedActorsVisualizationDesc& InOutVisualizationDesc);
 
 	// Register additional / alternate VisualizationDesc for instances to switch to, creating ISMC's
 	// for each VisualizationDesc.InstancedMeshes
