PR #12789: InstancedActors: Instanced Actor Component fixes

The CL fixed standalone-game IA uses by properly talking to ServerInstancedActorsSpawnerSubsystem, which fixes some IA representation issues.
Also, ServerInstancedActorsSpawnerSubsystem's been fixed to `FinishSpawning` actors only once their IAComponent has been configured.

diff --git a/Engine/Plugins/Runtime/InstancedActors/Source/InstancedActors/Private/InstancedActorsComponent.cpp b/Engine/Plugins/Runtime/InstancedActors/Source/InstancedActors/Private/InstancedActorsComponent.cpp
--- a/Engine/Plugins/Runtime/InstancedActors/Source/InstancedActors/Private/InstancedActorsComponent.cpp
+++ b/Engine/Plugins/Runtime/InstancedActors/Source/InstancedActors/Private/InstancedActorsComponent.cpp
@@ -48,7 +48,7 @@ void UInstancedActorsComponent::InitializeComponent()
 
 	// @todo Add support for non-replay NM_Standalone where we should call OnInstancedActorComponentInitialize
 	UWorld* World = GetWorld();
-	if (World && World->GetNetMode() == NM_DedicatedServer)
+	if (World && World->GetNetMode() != NM_Client)
 	{
 		UServerInstancedActorsSpawnerSubsystem* ServerInstancedActorSpawnerSubystem = UE::InstancedActors::Utils::GetServerInstancedActorsSpawnerSubsystem(*World);
 		if (IsValid(ServerInstancedActorSpawnerSubystem))
@@ -61,7 +61,7 @@ void UInstancedActorsComponent::InitializeComponent()
 void UInstancedActorsComponent::EndPlay(const EEndPlayReason::Type EndPlayReason)
 {
 	// Deregister Actor from entity on clients
-	if (!IsNetMode(NM_DedicatedServer) && InstanceHandle.IsValid())
+	if (GetNetMode() == NM_Client && InstanceHandle.IsValid())
 	{
 		AActor* Owner = GetOwner();
 		if (ensure(Owner))
diff --git a/Engine/Plugins/Runtime/InstancedActors/Source/InstancedActors/Private/ServerInstancedActorsSpawnerSubsystem.cpp b/Engine/Plugins/Runtime/InstancedActors/Source/InstancedActors/Private/ServerInstancedActorsSpawnerSubsystem.cpp
--- a/Engine/Plugins/Runtime/InstancedActors/Source/InstancedActors/Private/ServerInstancedActorsSpawnerSubsystem.cpp
+++ b/Engine/Plugins/Runtime/InstancedActors/Source/InstancedActors/Private/ServerInstancedActorsSpawnerSubsystem.cpp
@@ -68,6 +68,11 @@ ESpawnRequestStatus UServerInstancedActorsSpawnerSubsystem::SpawnActor(FConstStr
 	};
 
 	InOutSpawnParameters.SpawnCollisionHandlingOverride = ESpawnActorCollisionHandlingMethod::AlwaysSpawn;
+	// we’re going to call FinishSpawning only if the input parameters don’t indicate that the caller wants to handle it themselves.
+	const bool bCallFinishSpawning = (InOutSpawnParameters.bDeferConstruction == false);
+	// we always defer construction to have a chance to configure the UInstancedActorsComponent instances
+	// before their InitializeComponent gets called. From the callers point of view nothing changes.
+	InOutSpawnParameters.bDeferConstruction = true;
 
 	OutSpawnedActor = World->SpawnActor<AActor>(SpawnRequest.Template, SpawnRequest.Transform, InOutSpawnParameters);
 	if (ensureMsgf(OutSpawnedActor, TEXT("Failed to spawn actor of class %s"), *GetNameSafe(SpawnRequest.Template.Get())))
@@ -100,6 +105,11 @@ ESpawnRequestStatus UServerInstancedActorsSpawnerSubsystem::SpawnActor(FConstStr
 			InstancedActorComponent->SetIsReplicated(true);
 			InstancedActorComponent->RegisterComponent();
 		}
+
+		if (bCallFinishSpawning)
+		{
+			OutSpawnedActor->FinishSpawning(SpawnRequest.Transform);
+		}
 	}
 	
 	return IsValid(OutSpawnedActor) ? ESpawnRequestStatus::Succeeded : ESpawnRequestStatus::Failed;
