[InstancedActors] Moving UInstancedActorsData's BulkLOD and tick data to a struct cached in the UInstancedActorsSubsystem, instead of being stored as a FMassSharedFragment. See prior comments around FInstancedActorsDataSharedFragment: this was really not being used in the canonical shared fragment context, and with this change we can also better control lifetime if needed

diff --git a/Engine/Plugins/Runtime/InstancedActors/Source/InstancedActors/Private/InstancedActorsData.cpp b/Engine/Plugins/Runtime/InstancedActors/Source/InstancedActors/Private/InstancedActorsData.cpp
--- a/Engine/Plugins/Runtime/InstancedActors/Source/InstancedActors/Private/InstancedActorsData.cpp
+++ b/Engine/Plugins/Runtime/InstancedActors/Source/InstancedActors/Private/InstancedActorsData.cpp
@@ -402,20 +402,7 @@ void UInstancedActorsData::SpawnEntities()
 
 	// Now we've seeded Mass entity locations, we can free up now-superfluous InstanceTransforms
 	InstanceTransforms.Empty();
-
-	if (UE::InstancedActors::CVars::bUpdateNextTickTimeFragments)
-	{
-		FInstancedActorsDataSharedFragment* AsShared = SharedInstancedActorDataStruct.GetPtr<FInstancedActorsDataSharedFragment>();
-		// our shared fragment was being ticked in the past, so it's possible it's scheduled to tick in quite some time
-		// while we need the update ASAP. We're letting the UInstancedActorsSubsystem know to reschedule it.
-		if (ensure(AsShared) && AsShared->LastTickTime > 0.0)
-		{
-			UInstancedActorsSubsystem* InstancedActorSubsystem = UE::InstancedActors::Utils::GetInstancedActorsSubsystem(*World);
-			check(InstancedActorSubsystem);
-			InstancedActorSubsystem->UpdateAndResetTickTime(SharedInstancedActorDataStruct);
-		}
-	}
-
+	
 	// only once we spawned entities we can apply the "SetReplicatedActor" that came in before this InstancedActorData was ready to receive them.
 	for (const FSetReplicatedActorRequests& Request : CachedSetReplicatedActorRequests)
 	{
@@ -513,23 +500,6 @@ void UInstancedActorsData::DespawnEntities()
 		}
 	}
 
-	// Reset the related shared fragments. Not removing it fully since the system doesn't support it yet and we are
-	// going to reuse it anyway the next time this IA Manager gets loaded
-
-	FInstancedActorsDataSharedFragment ManagerFragment;
-	ManagerFragment.InstanceData = this;
-	FSharedStruct SharedFragmentInstance = MassEntityManager.GetOrCreateSharedFragment(ManagerFragment);
-
-	FInstancedActorsDataSharedFragment* AsSharedManagerFragment = SharedFragmentInstance.GetPtr<FInstancedActorsDataSharedFragment>();
-	if (ensure(AsSharedManagerFragment) && AsSharedManagerFragment->InstanceData == this)
-	{
-		// decrement the stats before resetting
-		Manager.UpdateInstanceStats(NumInstances, AsSharedManagerFragment->BulkLOD, /*Increment=*/false);
-
-		AsSharedManagerFragment->InstanceData = nullptr;
-		AsSharedManagerFragment->BulkLOD = EInstancedActorsBulkLOD::MAX;
-	}
-
 	DEC_DWORD_STAT_BY(STAT_TotalInstanceCount, NumInstances);
 
 	// Reregister & destroy runtime ISMCs
@@ -764,16 +734,6 @@ FMassEntityHandle UInstancedActorsData::GetEntity(FInstancedActorsInstanceIndex
 	return FMassEntityHandle();
 }
 
-void UInstancedActorsData::SetSharedInstancedActorDataStruct(FSharedStruct InSharedStruct)
-{
-	checkf(InSharedStruct.GetPtr<FInstancedActorsDataSharedFragment>(), TEXT("We expect only FInstancedActorsDataSharedFragment-base types here"));
-	checkf(SharedInstancedActorDataStruct.IsValid() == false 
-		|| InSharedStruct.GetMemory() == SharedInstancedActorDataStruct.GetMemory()
-		, TEXT("We don't expect to override this value"));
-
-	SharedInstancedActorDataStruct = TStructView<FInstancedActorsDataSharedFragment>(InSharedStruct.GetMemory());
-}
-
 FString UInstancedActorsData::GetDebugName(bool bCompact) const
 {
 	if (bCompact)
diff --git a/Engine/Plugins/Runtime/InstancedActors/Source/InstancedActors/Private/InstancedActorsManager.cpp b/Engine/Plugins/Runtime/InstancedActors/Source/InstancedActors/Private/InstancedActorsManager.cpp
--- a/Engine/Plugins/Runtime/InstancedActors/Source/InstancedActors/Private/InstancedActorsManager.cpp
+++ b/Engine/Plugins/Runtime/InstancedActors/Source/InstancedActors/Private/InstancedActorsManager.cpp
@@ -357,6 +357,8 @@ void AInstancedActorsManager::InitializeModifyAndSpawnEntities()
 		check(InstanceData);
 		InstanceData->Initialize();
 	}
+	
+	InstancedActorSubsystem->RegisterInstancedActorsData(PerActorClassInstanceData);
 
 	// Run pending modifiers that can run pre-entity spawn
 	TryRunPendingModifiers();
@@ -390,6 +392,8 @@ void AInstancedActorsManager::EndPlay(const EEndPlayReason::Type EndPlayReason)
 		check(InstanceData);
 		InstanceData->Deinitialize();
 	}
+	
+	InstancedActorSubsystem->UnregisterInstancedActorsData(PerActorClassInstanceData);
 
 	// Deregister with UInstancedActorsSubsystem
 	if (InstancedActorSubsystem)
@@ -1125,9 +1129,9 @@ bool AInstancedActorsManager::HasInstancesOfClass(const FBox& InQueryBounds, TSu
 	FScopedInstancedActorsIterationContext IterationContext;
 	bool bHasInstance = false;
 
-	auto InstancedActorDataMask = TOptional<AInstancedActorsManager::FInstancedActorDataPredicateFunc>([=](const UInstancedActorsData& InstancedActorData)
+	auto InstancedActorDataMask = TOptional<AInstancedActorsManager::FInstancedActorDataPredicateFunc>([=, InstancedActorSubsystem = InstancedActorSubsystem](const UInstancedActorsData& InstancedActorData)
 		{
-			return ((int(AllowedLODs) & (1 << int(InstancedActorData.GetBulkLOD()))) != 0)
+			return ((int(AllowedLODs) & (1 << int(InstancedActorSubsystem->GetInstancedActorsDataBulkLOD(InstancedActorData)))) != 0)
 				&& InstancedActorData.ActorClass->IsChildOf(ActorClass);
 		});
 
diff --git a/Engine/Plugins/Runtime/InstancedActors/Source/InstancedActors/Private/InstancedActorsStationaryLODBatchProcessor.cpp b/Engine/Plugins/Runtime/InstancedActors/Source/InstancedActors/Private/InstancedActorsStationaryLODBatchProcessor.cpp
--- a/Engine/Plugins/Runtime/InstancedActors/Source/InstancedActors/Private/InstancedActorsStationaryLODBatchProcessor.cpp
+++ b/Engine/Plugins/Runtime/InstancedActors/Source/InstancedActors/Private/InstancedActorsStationaryLODBatchProcessor.cpp
@@ -226,10 +226,10 @@ void UInstancedActorsStationaryLODBatchProcessor::Execute(FMassEntityManager& En
 		auto ExecutionFunction = [Viewers = MakeArrayView((const FViewerInfo*)&Viewers[0], Viewers.Num()), &EntityManager, &Context
 			, LODChangingEntityQuery = &LODChangingEntityQuery, StaticMeshLODDistanceScale, CurrentTime
 			, DelayPerBulkLOD = MakeArrayView((const double*)&DelayPerBulkLOD[0], (int)EInstancedActorsBulkLOD::MAX)]
-			(FInstancedActorsDataSharedFragment& ManagerSharedFragment) -> double
+			(UInstancedActorsSubsystem::FInstancedActorsDataTickInfo& InstancedActorsDataTickInfo) -> double
 			{
 				double NextTickTime = CurrentTime + DelayPerBulkLOD[(int)EInstancedActorsBulkLOD::Off];
-				if (UInstancedActorsData* InstanceData = ManagerSharedFragment.InstanceData.Get())
+				if (UInstancedActorsData* InstanceData = InstancedActorsDataTickInfo.InstancedActorsData.Get())
 				{
 					const FInstancedActorsSettings& Settings = InstanceData->GetSettings<const FInstancedActorsSettings>();
 					
@@ -284,13 +284,13 @@ void UInstancedActorsStationaryLODBatchProcessor::Execute(FMassEntityManager& En
 					// Updates the time at which the FInstancedActorsDataSharedFragment will tick depending on its bulk LOD value
 					NextTickTime = CurrentTime + (DelayPerBulkLOD[(int)NewBulkLOD] * 0.95 + FMath::FRand() * 0.1);
 
-					if (const bool bHasBulkLODChanged = (ManagerSharedFragment.BulkLOD != NewBulkLOD))
+					if (const bool bHasBulkLODChanged = (InstancedActorsDataTickInfo.BulkLOD != NewBulkLOD))
 					{
 						// Decrements stats with current state, and then increments stats with the new one
 						{
-							AInstancedActorsManager::UpdateInstanceStats(InstanceData->NumInstances, ManagerSharedFragment.BulkLOD, false);
-							ManagerSharedFragment.BulkLOD = NewBulkLOD;
-							AInstancedActorsManager::UpdateInstanceStats(InstanceData->NumInstances, ManagerSharedFragment.BulkLOD, true);
+							AInstancedActorsManager::UpdateInstanceStats(InstanceData->NumInstances, InstancedActorsDataTickInfo.BulkLOD, false);
+							InstancedActorsDataTickInfo.BulkLOD = NewBulkLOD;
+							AInstancedActorsManager::UpdateInstanceStats(InstanceData->NumInstances, InstancedActorsDataTickInfo.BulkLOD, true);
 						}
 						// Toggles physics state for the IA's ISM depending on the new bulk LOD value.
 						// If enabled = physics on, else = physics off.				
@@ -339,7 +339,7 @@ void UInstancedActorsStationaryLODBatchProcessor::Execute(FMassEntityManager& En
 						}
 						// Toggles MassProcessors on/off depending on the bulk LOD, by pushing or removing tags that are used by those processor's queries.
 						// NOTE: Forcibly updates the mass LOD to off or low when bulk LOD is smaller than Detailed
-						if (ManagerSharedFragment.BulkLOD == EInstancedActorsBulkLOD::Detailed && InstanceData->CanHydrate())
+						if (InstancedActorsDataTickInfo.BulkLOD == EInstancedActorsBulkLOD::Detailed && InstanceData->CanHydrate())
 						{
 							EntityManager.Defer().PushCommand<UE::InstancedActors::FEnableDetailedLODCommand>(InstanceData->Entities);
 						}
@@ -347,7 +347,7 @@ void UInstancedActorsStationaryLODBatchProcessor::Execute(FMassEntityManager& En
 						{
 							// Force given LOD for all the hosted entities 
 							EMassLOD::Type NewLOD = EMassLOD::Off;
-							switch (ManagerSharedFragment.BulkLOD)
+							switch (InstancedActorsDataTickInfo.BulkLOD)
 							{
 							case EInstancedActorsBulkLOD::Detailed:
 								ensureMsgf(InstanceData->CanHydrate() == false, TEXT("This case is only valid for non-hydrating instance, broken for %s"), *GetNameSafe(InstanceData->ActorClass));
@@ -389,18 +389,28 @@ void UInstancedActorsStationaryLODBatchProcessor::Execute(FMassEntityManager& En
 				return NextTickTime;
 			};
 
-		TArray<UInstancedActorsSubsystem::FNextTickSharedFragment>& SortedSharedFragments = InstancedActorSubsystem->GetTickableSharedFragments();
-		if (SortedSharedFragments.Num() > 0)
+		InstancedActorSubsystem->SortInstancedActorsDataTickInfo();
+		while (!InstancedActorSubsystem->InstancedActorsDataTickInfoIsEmpty())
 		{
-			while (SortedSharedFragments.HeapTop().NextTickTime < CurrentTime)
+			UInstancedActorsSubsystem::FInstancedActorsDataTickInfo* InstancedActorsDataTickInfo = InstancedActorSubsystem->PeekInstancedActorsDataTickInfoTop();
+			double NextTickTime = 0.0;
+			bool bShouldExit = false;
+			if (InstancedActorsDataTickInfo != nullptr)
 			{
-				UInstancedActorsSubsystem::FNextTickSharedFragment WrappedSharedFragment;
-				SortedSharedFragments.HeapPop(WrappedSharedFragment, EAllowShrinking::No);
-				FInstancedActorsDataSharedFragment& ManagerSharedFragment = WrappedSharedFragment.SharedStruct.Get<FInstancedActorsDataSharedFragment>();
+				if (InstancedActorsDataTickInfo->NextTickTime >= CurrentTime)
+				{
+					bShouldExit = true;
+				}
 				
-				ManagerSharedFragment.LastTickTime = CurrentTime;
-				WrappedSharedFragment.NextTickTime = ExecutionFunction(ManagerSharedFragment);
-				SortedSharedFragments.HeapPush(MoveTemp(WrappedSharedFragment));
+				NextTickTime = ExecutionFunction(*InstancedActorsDataTickInfo);
+			}
+
+			// this will simply pop the heap's top if references are invalid, so we still need to call it
+			InstancedActorSubsystem->UpdateInstancedActorsDataTickInfoTop(CurrentTime, NextTickTime);
+
+			if (bShouldExit)
+			{
+				break;
 			}
 		}
 
@@ -425,7 +435,7 @@ void UInstancedActorsStationaryLODBatchProcessor::Execute(FMassEntityManager& En
 
 				if (DirtyRepresentationInstance.IsValid() 
 					// only the entities that are not "Detailed" require update, so we're filtering out accordingly
-					&& DirtyRepresentationInstance.GetInstanceActorDataChecked().GetBulkLOD() > EInstancedActorsBulkLOD::Detailed)
+					&& InstancedActorSubsystem->GetInstancedActorsDataBulkLOD(DirtyRepresentationInstance.GetInstanceActorDataChecked()) > EInstancedActorsBulkLOD::Detailed)
 				{
 					FMassEntityHandle DirtyEntity = DirtyRepresentationInstance.GetInstanceActorDataChecked().GetEntity(DirtyRepresentationInstance.GetInstanceIndex());
 					if (EntityManager.IsEntityValid(DirtyEntity))
diff --git a/Engine/Plugins/Runtime/InstancedActors/Source/InstancedActors/Private/InstancedActorsSubsystem.cpp b/Engine/Plugins/Runtime/InstancedActors/Source/InstancedActors/Private/InstancedActorsSubsystem.cpp
--- a/Engine/Plugins/Runtime/InstancedActors/Source/InstancedActors/Private/InstancedActorsSubsystem.cpp
+++ b/Engine/Plugins/Runtime/InstancedActors/Source/InstancedActors/Private/InstancedActorsSubsystem.cpp
@@ -10,6 +10,7 @@
 #include "InstancedActorsSettings.h"
 #include "ActorPartition/ActorPartitionSubsystem.h"
 #include "Algo/Find.h"
+#include "Algo/IndexOf.h"
 #include "DataRegistry.h"
 #include "DataRegistrySubsystem.h"
 #include "Engine/Engine.h"
@@ -85,6 +86,22 @@ namespace InstancedActorsCVars
 #endif
 }
 
+namespace UE::InstancedActors
+{
+	struct FSortInstancedActorsDataTickInfoPredicate
+	{
+		FSortInstancedActorsDataTickInfoPredicate(TArray<UInstancedActorsSubsystem::FInstancedActorsDataTickInfo>& InInstancedActorsDataInfo)
+			: InstancedActorsDataInfo(InInstancedActorsDataInfo)
+		{}
+
+		bool operator()(int32 A, int32 B) const
+		{
+			return InstancedActorsDataInfo[A].NextTickTime < InstancedActorsDataInfo[B].NextTickTime;
+		}
+		TArray<UInstancedActorsSubsystem::FInstancedActorsDataTickInfo>& InstancedActorsDataInfo;
+	};
+}
+
 //-----------------------------------------------------------------------------
 // UInstancedActorsSubsystem
 //-----------------------------------------------------------------------------
@@ -900,60 +917,135 @@ FInstancedActorsVisualizationDesc UInstancedActorsSubsystem::CreateVisualDescrip
 	return FInstancedActorsVisualizationDesc::FromActor(ExemplarActor);
 }
 
-TArray<UInstancedActorsSubsystem::FNextTickSharedFragment>& UInstancedActorsSubsystem::GetTickableSharedFragments()
-{	
-	RegisterNewSharedFragmentsInternal();
-	return SortedSharedFragments;
+void UInstancedActorsSubsystem::SortInstancedActorsDataTickInfo()
+{
+	if (bInstancedActorsDataInfoDirty)
+	{
+		TRACE_CPUPROFILER_EVENT_SCOPE(UInstancedActorsSubsystem_SortInstancedActorsDataTickInfo);
+		Algo::Heapify(SortedInstancedActorsDataInfo, UE::InstancedActors::FSortInstancedActorsDataTickInfoPredicate(InstancedActorsDataInfo));
+		bInstancedActorsDataInfoDirty = false;
+	}
 }
 
-void UInstancedActorsSubsystem::UpdateAndResetTickTime(TConstStructView<FInstancedActorsDataSharedFragment> InstancedActorsDataSharedFragment)
+UInstancedActorsSubsystem::FInstancedActorsDataTickInfo* UInstancedActorsSubsystem::PeekInstancedActorsDataTickInfoTop()
 {
-	TRACE_CPUPROFILER_EVENT_SCOPE(UInstancedActorsSubsystem_UpdateTickableSharedFragments);
+	if (SortedInstancedActorsDataInfo.IsEmpty())
+	{
+		return nullptr;
+	}
+	
+	return &InstancedActorsDataInfo[SortedInstancedActorsDataInfo.HeapTop()];
+}
 
-	const bool bParamStructFound = RegisterNewSharedFragmentsInternal(InstancedActorsDataSharedFragment);
-	if (bParamStructFound == false)
+void UInstancedActorsSubsystem::UpdateInstancedActorsDataTickInfoTop(double CurrentTime, double NextTickTime)
+{
+	if (SortedInstancedActorsDataInfo.IsEmpty())
 	{
-		// we need to find the one. Naive implementation for now.
-		// starting from back in assumption stuff has been removed and re-added so the relevant data should be closer 
-		// to the back due to being at EMassLOD::Off level
-		for (int32 Index = SortedSharedFragments.Num() - 1; Index >= 0; --Index)
+		return;
+	}
+	
+	TRACE_CPUPROFILER_EVENT_SCOPE(UInstancedActorsSubsystem_UpdateInstancedActorsDataTickInfoTop);
+
+	// heap pop
+	const int32 TickInfoIndex = SortedInstancedActorsDataInfo[0];
+	if (ensure(InstancedActorsDataInfo.IsValidIndex(TickInfoIndex)))
+	{
+		FInstancedActorsDataTickInfo& TickInfo = InstancedActorsDataInfo[TickInfoIndex];
+		SortedInstancedActorsDataInfo.RemoveAtSwap(0, 1, EAllowShrinking::No);
+		AlgoImpl::HeapSiftDown(SortedInstancedActorsDataInfo.GetData(), 0, SortedInstancedActorsDataInfo.Num(), FIdentityFunctor()
+			, UE::InstancedActors::FSortInstancedActorsDataTickInfoPredicate(InstancedActorsDataInfo));
+
+		if (RegisteredInstancedActorsDataMap.Contains(TickInfo.InstancedActorsData.Get()))
 		{
-			if (SortedSharedFragments[Index].SharedStruct == InstancedActorsDataSharedFragment)
-			{
-				FNextTickSharedFragment TickFragment = SortedSharedFragments[Index];
-				// setting to 0 will force update the very next time Batch LOD is being calculated. 
-				TickFragment.NextTickTime = 0;
-				SortedSharedFragments.HeapRemoveAt(Index, EAllowShrinking::No);
-				SortedSharedFragments.HeapPush(MoveTemp(TickFragment));
-				break;
-			}
+			TickInfo.LastTickTime = CurrentTime;
+			TickInfo.NextTickTime = NextTickTime;
+
+			// heap push
+			SortedInstancedActorsDataInfo.Add(TickInfoIndex);
+			AlgoImpl::HeapSiftUp(SortedInstancedActorsDataInfo.GetData(), 0, SortedInstancedActorsDataInfo.Num() - 1, FIdentityFunctor()
+				, UE::InstancedActors::FSortInstancedActorsDataTickInfoPredicate(InstancedActorsDataInfo));
 		}
 	}
 }
 
-bool UInstancedActorsSubsystem::RegisterNewSharedFragmentsInternal(TConstStructView<FInstancedActorsDataSharedFragment> InstancedActorsDataSharedFragment)
+const UInstancedActorsSubsystem::FInstancedActorsDataTickInfo* UInstancedActorsSubsystem::GetInstancedActorsDataTickInfo(const UInstancedActorsData& InstancedActorsData) const
+{
+	TRACE_CPUPROFILER_EVENT_SCOPE(UInstancedActorsSubsystem_GetInstancedActorsDataTickInfo);
+	
+	const int32* FoundIndexPtr = RegisteredInstancedActorsDataMap.Find(&InstancedActorsData);
+	if (FoundIndexPtr == nullptr)
+	{
+		return nullptr;
+	}
+	
+	return &InstancedActorsDataInfo[*FoundIndexPtr];
+}
+
+EInstancedActorsBulkLOD UInstancedActorsSubsystem::GetInstancedActorsDataBulkLOD(const UInstancedActorsData& InstancedActorsData) const
 {
-	check(EntityManager);
+	if (const FInstancedActorsDataTickInfo* InstancedActorsDataTickInfo = GetInstancedActorsDataTickInfo(InstancedActorsData))
+	{
+		return InstancedActorsDataTickInfo->BulkLOD;
+	}
 
-	const bool bParamStructProvided = InstancedActorsDataSharedFragment.IsValid();
-	// starting with !bParamStructProvided will result in short-circuiting the assignment operation below if InstancedActorsDataSharedFragment is empty
-	bool bParamStructFound = !bParamStructProvided;
-	TConstArrayView<FSharedStruct> AllSharedFragmentsOfType = EntityManager->GetSharedFragmentsOfType<FInstancedActorsDataSharedFragment>();
+	ensureMsgf(false, TEXT("We expect to hit this path only if the IAD instance has not been initialized yet."));
+	return EInstancedActorsBulkLOD::MAX;
+}
+
+void UInstancedActorsSubsystem::RegisterInstancedActorsData(TArray<TObjectPtr<UInstancedActorsData>> InstancedActorsDataArray)
+{
+	TRACE_CPUPROFILER_EVENT_SCOPE(UInstancedActorsSubsystem_RegisterInstancedActorsData);
 	
-	if (SortedSharedFragments.Num() < AllSharedFragmentsOfType.Num())
+	for (TObjectPtr<UInstancedActorsData>& InstancedActorsDataPtr : InstancedActorsDataArray)
 	{
-		// We add all of them at the front for immediate processing.
-		const int32 StartingIndex = SortedSharedFragments.Num();
-		const int32 NewItemsCount = (AllSharedFragmentsOfType.Num() - SortedSharedFragments.Num());
-		SortedSharedFragments.InsertDefaulted(0, NewItemsCount);
-		for (int32 NewIndex = 0; NewIndex < NewItemsCount; ++NewIndex)
+		UInstancedActorsData* InstancedActorsData = InstancedActorsDataPtr.Get();		
+		if (InstancedActorsData == nullptr)
 		{
-			SortedSharedFragments[NewIndex].SharedStruct = AllSharedFragmentsOfType[StartingIndex + NewIndex];
-			bParamStructFound = bParamStructFound || (AllSharedFragmentsOfType[StartingIndex + NewIndex] == InstancedActorsDataSharedFragment);
+			continue;
+		}
+	
+		if (RegisteredInstancedActorsDataMap.Contains(InstancedActorsData))
+		{
+			continue;
 		}
-		SortedSharedFragments.Heapify();
-	}
-	check(SortedSharedFragments.Num() == AllSharedFragmentsOfType.Num());
 
-	return bParamStructFound;
+		int32 Index = INDEX_NONE;
+		if (InstancedActorsDataInfoFreeIndicesStack.Num())
+		{
+			Index = InstancedActorsDataInfoFreeIndicesStack.Pop(EAllowShrinking::No);
+
+			InstancedActorsDataInfo[Index] = {.InstancedActorsData = InstancedActorsData};
+		}
+		else
+		{
+			Index = InstancedActorsDataInfo.Add({.InstancedActorsData = InstancedActorsData});
+		}
+
+		RegisteredInstancedActorsDataMap.Add(InstancedActorsData, Index);
+		SortedInstancedActorsDataInfo.Push(Index);
+		
+		bInstancedActorsDataInfoDirty = true;
+	}	
 }
+
+void UInstancedActorsSubsystem::UnregisterInstancedActorsData(TArray<TObjectPtr<UInstancedActorsData>> InstancedActorsDataArray)
+{
+	TRACE_CPUPROFILER_EVENT_SCOPE(UInstancedActorsSubsystem_UnregisterInstancedActorsData);
+	
+	for (TObjectPtr<UInstancedActorsData>& InstancedActorsDataPtr : InstancedActorsDataArray)
+	{
+		UInstancedActorsData* InstancedActorsData = InstancedActorsDataPtr.Get();
+		if (InstancedActorsData == nullptr)
+		{
+			continue;
+		}
+
+		int32 StoredIndex = INDEX_NONE;
+		if (RegisteredInstancedActorsDataMap.RemoveAndCopyValue(InstancedActorsData, StoredIndex))
+		{
+			InstancedActorsDataInfoFreeIndicesStack.Add(StoredIndex);
+		}
+
+		// no need to mark the sorted array as dirty. Invalid references will be checked at consumption time and not re-heapified
+	}
+}
\ No newline at end of file
diff --git a/Engine/Plugins/Runtime/InstancedActors/Source/InstancedActors/Private/InstancedActorsVisualizationTrait.cpp b/Engine/Plugins/Runtime/InstancedActors/Source/InstancedActors/Private/InstancedActorsVisualizationTrait.cpp
--- a/Engine/Plugins/Runtime/InstancedActors/Source/InstancedActors/Private/InstancedActorsVisualizationTrait.cpp
+++ b/Engine/Plugins/Runtime/InstancedActors/Source/InstancedActors/Private/InstancedActorsVisualizationTrait.cpp
@@ -91,33 +91,6 @@ void UInstancedActorsVisualizationTrait::BuildTemplate(FMassEntityTemplateBuildC
 	BuildContext.AddFragment<FTransformFragment>();
 	BuildContext.AddFragment<FMassActorFragment>();
 
-	FInstancedActorsDataSharedFragment ManagerSharedFragment;
-	ManagerSharedFragment.InstanceData = InstanceData;
-	FSharedStruct SubsystemFragment = EntityManager.GetOrCreateSharedFragment(ManagerSharedFragment);
-
-	FInstancedActorsDataSharedFragment* AsShared = SubsystemFragment.GetPtr<FInstancedActorsDataSharedFragment>();
-	if (ensure(AsShared))
-	{
-		// this can happen when we unload data and then stream it back again - we end up with the same path object, but different pointer.
-		// The old instance should be garbage now
-		if (AsShared->InstanceData != InstanceData)
-		{
-			ensure(AsShared->InstanceData.IsValid() == false);
-			AsShared->InstanceData = InstanceData;
-		}
-		
-		// we also need to make sure the bulk LOD is reset since the shared fragment can survive the death of the original 
-		// InstanceData while preserving the "runtime" value, which will mess up newly spawned entities.
-		AsShared->BulkLOD = EInstancedActorsBulkLOD::MAX;
-
-		if (BuildContext.IsInspectingData() == false)
-		{
-			InstanceData->SetSharedInstancedActorDataStruct(SubsystemFragment);
-		}
-	}
-	// not adding SubsystemFragment do BuildContext on purpose, we temporarily use shared fragments to store IAD information.
-	// To be moved to InstancedActorSubsystem in the future
-
 	// ActorInstanceFragment will get initialized by UInstancedActorsInitializerProcessor
 	BuildContext.AddFragment_GetRef<FMassActorInstanceFragment>();
 
diff --git a/Engine/Plugins/Runtime/InstancedActors/Source/InstancedActors/Public/InstancedActorsData.h b/Engine/Plugins/Runtime/InstancedActors/Source/InstancedActors/Public/InstancedActorsData.h
--- a/Engine/Plugins/Runtime/InstancedActors/Source/InstancedActors/Public/InstancedActorsData.h
+++ b/Engine/Plugins/Runtime/InstancedActors/Source/InstancedActors/Public/InstancedActorsData.h
@@ -76,19 +76,6 @@ public:
 
 	UE_API FMassEntityHandle GetEntity(FInstancedActorsInstanceIndex InstanceIndex) const;
 
-	UE_API void SetSharedInstancedActorDataStruct(FSharedStruct InSharedStruct);
-	
-	/** @return current Bulk LOD of this InstancedActorData instance */
-	EInstancedActorsBulkLOD GetBulkLOD() const
-	{
-		if (const FInstancedActorsDataSharedFragment* DataFragment = SharedInstancedActorDataStruct.GetPtr<FInstancedActorsDataSharedFragment>())
-		{
-			return DataFragment->BulkLOD;
-		}
-		ensureMsgf(bHasEverInitialized == false, TEXT("We expect to hit this path only if the IAD instance has not been initialized yet."));
-		return EInstancedActorsBulkLOD::MAX;
-	}
-
 	// Identifying integer, unique within outer IAM, used for identifying matching persistence records
 	// Incrementally assigned in GetOrCreateActorInstanceData
 	UPROPERTY(VisibleAnywhere, Category=InstancedActors)
@@ -496,9 +483,6 @@ protected:
 	FBox CachedLocalBounds = FBox(ForceInit);
 
 private:
-	// Represents the shared fragment registered with MassEntityManager, that points back to this UInstancedActorsData instance
-	TStructView<FInstancedActorsDataSharedFragment> SharedInstancedActorDataStruct;
-
 	// List of FInstancedActorsDelta's to apply to instances, replicated via fast array replication to clients.
 	// Formed at runtime via player actions and persisted / restored by game's persistence system
 	// @see Serialize
diff --git a/Engine/Plugins/Runtime/InstancedActors/Source/InstancedActors/Public/InstancedActorsSubsystem.h b/Engine/Plugins/Runtime/InstancedActors/Source/InstancedActors/Public/InstancedActorsSubsystem.h
--- a/Engine/Plugins/Runtime/InstancedActors/Source/InstancedActors/Public/InstancedActorsSubsystem.h
+++ b/Engine/Plugins/Runtime/InstancedActors/Source/InstancedActors/Public/InstancedActorsSubsystem.h
@@ -177,7 +177,7 @@ public:
 	/*
 	* Called when an additional/alternate VisualizationDesc is registered. Override to make custom modifications to the visual representation
 	*/
-	virtual void ModifyVisualDescriptionForActor(const TNotNull<AActor*> ExemplarActor, FInstancedActorsVisualizationDesc& InOutVisualization) const {}
+	INSTANCEDACTORS_API virtual void ModifyVisualDescriptionForActor(const TNotNull<AActor*> ExemplarActor, FInstancedActorsVisualizationDesc& InOutVisualization) const {}
 
 	//~ Begin UTickableWorldSubsystem Overrides
 	INSTANCEDACTORS_API virtual void Tick(float DeltaTime) override;
@@ -190,19 +190,30 @@ public:
 	INSTANCEDACTORS_API virtual void Deinitialize() override;
 	//~ End USubsystem Overrides
 
-	struct FNextTickSharedFragment
+	struct FInstancedActorsDataTickInfo
 	{
-		FSharedStruct SharedStruct;
-		double NextTickTime = 0;
+		TWeakObjectPtr<UInstancedActorsData> InstancedActorsData;		
 
-		bool operator<(const FNextTickSharedFragment& Other) const
-		{
-			return NextTickTime < Other.NextTickTime;
-		}
+		EInstancedActorsBulkLOD BulkLOD = EInstancedActorsBulkLOD::MAX;
+
+		double LastTickTime = 0.0;
+		double NextTickTime = 0.0;
 	};
 
-	TArray<UInstancedActorsSubsystem::FNextTickSharedFragment>& GetTickableSharedFragments();
-	void UpdateAndResetTickTime(TConstStructView<FInstancedActorsDataSharedFragment> InstancedActorsDataSharedFragment);
+	void SortInstancedActorsDataTickInfo();
+	bool InstancedActorsDataTickInfoIsEmpty() const
+	{
+		return SortedInstancedActorsDataInfo.IsEmpty();
+	}
+	
+	FInstancedActorsDataTickInfo* PeekInstancedActorsDataTickInfoTop();
+	void UpdateInstancedActorsDataTickInfoTop(double CurrentTime, double NextTickTime);
+	const FInstancedActorsDataTickInfo* GetInstancedActorsDataTickInfo(const UInstancedActorsData& InstancedActorsData) const;
+	INSTANCEDACTORS_API EInstancedActorsBulkLOD GetInstancedActorsDataBulkLOD(const UInstancedActorsData& InstancedActorsData) const;
+	
+	// operates on TArray for heap efficiency
+	void RegisterInstancedActorsData(TArray<TObjectPtr<UInstancedActorsData>> InstancedActorsDataArray);
+	void UnregisterInstancedActorsData(TArray<TObjectPtr<UInstancedActorsData>> InstancedActorsDataArray);
 
 	TSubclassOf<AInstancedActorsManager> GetInstancedActorsManagerClass() const 
 	{ 
@@ -210,16 +221,17 @@ public:
 	}
 
 protected:
-	/** 
-	 * Fetches all registered FInstancedActorsDataSharedFragment from the EntityManager and adds the missing ones to SortedSharedFragments
-	 * @param InstancedActorsDataSharedFragment optionally the function can check if given shared fragment is amongst 
-	 *	the newly added fragments
-	 * @param returns whether InstancedActorsDataSharedFragment has been found, or `true` if that param is not provided. 
-	 */
-	bool RegisterNewSharedFragmentsInternal(TConstStructView<FInstancedActorsDataSharedFragment> InstancedActorsDataSharedFragment = TConstStructView<FInstancedActorsDataSharedFragment>());
+	/** Container hosting tick data per IAData instance. */
+	TArray<FInstancedActorsDataTickInfo> InstancedActorsDataInfo;
+	/** Contains indices to InstancedActorsDataInfo that have been released and can be reused*/
+	TArray<int32> InstancedActorsDataInfoFreeIndicesStack;
+
+	/** Maps specific IAD instances to indices to InstancedActorsDataInfo */
+	TMap<TObjectKey<UInstancedActorsData>, int32> RegisteredInstancedActorsDataMap;
+	/** The container storing a sorted queue of indices of InstancedActorsDataInfo, ordered by the NextTickTime */
+	TArray<int32> SortedInstancedActorsDataInfo;
 
-	/** The container storing a sorted queue of FSharedStruct instances, ordered by the NextTickTime */
-	TArray<FNextTickSharedFragment> SortedSharedFragments;
+	bool bInstancedActorsDataInfoDirty = false;
 
 	TSharedPtr<FMassEntityManager> EntityManager;
 
diff --git a/Engine/Plugins/Runtime/InstancedActors/Source/InstancedActors/Public/InstancedActorsTypes.h b/Engine/Plugins/Runtime/InstancedActors/Source/InstancedActors/Public/InstancedActorsTypes.h
--- a/Engine/Plugins/Runtime/InstancedActors/Source/InstancedActors/Public/InstancedActorsTypes.h
+++ b/Engine/Plugins/Runtime/InstancedActors/Source/InstancedActors/Public/InstancedActorsTypes.h
@@ -318,24 +318,6 @@ private:
 	int32 ModifierVolumeID = INDEX_NONE;
 };
 
-/**
- * Note that we don't really need this type to be a shared fragment. It's used to create FSharedStructs pointing at
- * UInstancedActorsData and this data is fetched from MassEntityManager by UInstancedActorsStationaryLODBatchProcessor.
- * @todo This will be addressed in the future by refactoring where this data is stored and how it's used.
- */
-USTRUCT()
-struct FInstancedActorsDataSharedFragment : public FMassSharedFragment
-{
-	GENERATED_BODY()
-
-	UPROPERTY(Transient)
-	TWeakObjectPtr<UInstancedActorsData> InstanceData;
-
-	EInstancedActorsBulkLOD BulkLOD = EInstancedActorsBulkLOD::MAX;
-
-	double LastTickTime = 0.0;
-};
-
 USTRUCT()
 struct FInstancedActorsFragment : public FMassFragment
 {
