Add missing header that's only incidentally included in builds with mass debug enabled

diff --git a/Engine/Plugins/Runtime/InstancedActors/Source/InstancedActors/Public/InstancedActorsStationaryLODBatchProcessor.h b/Engine/Plugins/Runtime/InstancedActors/Source/InstancedActors/Public/InstancedActorsStationaryLODBatchProcessor.h
--- a/Engine/Plugins/Runtime/InstancedActors/Source/InstancedActors/Public/InstancedActorsStationaryLODBatchProcessor.h
+++ b/Engine/Plugins/Runtime/InstancedActors/Source/InstancedActors/Public/InstancedActorsStationaryLODBatchProcessor.h
@@ -2,6 +2,7 @@
 
 #pragma once
 
+#include "MassEntityQuery.h"
 #include "MassProcessor.h"
 #include "InstancedActorsTypes.h"
 #include "StructUtils/SharedStruct.h"
