Updated dll storage on code using UnrealCodeFixup

diff --git a/Engine/Plugins/Runtime/InstancedActors/Source/InstancedActors/Public/InstancedActorsSubsystem.h b/Engine/Plugins/Runtime/InstancedActors/Source/InstancedActors/Public/InstancedActorsSubsystem.h
--- a/Engine/Plugins/Runtime/InstancedActors/Source/InstancedActors/Public/InstancedActorsSubsystem.h
+++ b/Engine/Plugins/Runtime/InstancedActors/Source/InstancedActors/Public/InstancedActorsSubsystem.h
@@ -177,7 +177,7 @@ public:
 	/*
 	* Called when an additional/alternate VisualizationDesc is registered. Override to make custom modifications to the visual representation
 	*/
-	INSTANCEDACTORS_API virtual void ModifyVisualDescriptionForActor(const TNotNull<AActor*> ExemplarActor, FInstancedActorsVisualizationDesc& InOutVisualization) const {}
+	virtual void ModifyVisualDescriptionForActor(const TNotNull<AActor*> ExemplarActor, FInstancedActorsVisualizationDesc& InOutVisualization) const {}
 
 	//~ Begin UTickableWorldSubsystem Overrides
 	INSTANCEDACTORS_API virtual void Tick(float DeltaTime) override;
diff --git a/Engine/Source/Runtime/Engine/Classes/GameFramework/CharacterMovementComponent.h b/Engine/Source/Runtime/Engine/Classes/GameFramework/CharacterMovementComponent.h
--- a/Engine/Source/Runtime/Engine/Classes/GameFramework/CharacterMovementComponent.h
+++ b/Engine/Source/Runtime/Engine/Classes/GameFramework/CharacterMovementComponent.h
@@ -1595,7 +1595,7 @@ public:
 	ENGINE_API virtual void UpdateCharacterStateAfterMovement(float DeltaSeconds);
 
 	/** Give one final chance to update the velocity right before doing the actual position change. */
-	ENGINE_API virtual void UpdateVelocityBeforeMovement(float DeltaSeconds) {}
+	virtual void UpdateVelocityBeforeMovement(float DeltaSeconds) {}
 
 	/** 
 	 * Handle start swimming functionality
