[InstancedActors] Modified the occasionally failing ensure condition in UInstancedActorsData::OnInstancedActorDestroyed and provided it with more comments and a message to emit.

diff --git a/Engine/Plugins/Runtime/InstancedActors/Source/InstancedActors/Private/InstancedActorsData.cpp b/Engine/Plugins/Runtime/InstancedActors/Source/InstancedActors/Private/InstancedActorsData.cpp
--- a/Engine/Plugins/Runtime/InstancedActors/Source/InstancedActors/Private/InstancedActorsData.cpp
+++ b/Engine/Plugins/Runtime/InstancedActors/Source/InstancedActors/Private/InstancedActorsData.cpp
@@ -1686,9 +1686,12 @@ void UInstancedActorsData::OnInstancedActorDestroyed(AActor& DestroyedActor, con
 	{
 		// Before destroying the entity, thus insta-destroying the associated actor, we release Mass' reference
 		// to the actor here to prevent this from happening, rather allowing the actor to destroy itself as
-		// it is reporting to do so, allowing death FX etc to be played first
+		// it is reporting to do so, allowing death FX etc to be played first.
+		// Note that at this point it's possible for retrieved actor to be nullptr if the related actor is already considered destroyed. 
 		AActor* UnlinkedActor = UnlinkActor(DestroyedInstanceIndex);
-		ensure(UnlinkedActor == &DestroyedActor);
+		ensureMsgf(UnlinkedActor == &DestroyedActor || UnlinkedActor == nullptr
+			, TEXT("%hs: Unexpected Unlink result. Actor being destroyed: %s, actor unlinked %s")
+			, __FUNCTION__, *DestroyedActor.GetName(), *UnlinkedActor->GetName());
 
 		// If we got this far, destruction must have been caused by player, mark the instance for destruction.
 		// Note: This happens on both client & server.
