IAMs create unready volumes until instances are created

diff --git a/Engine/Plugins/Runtime/InstancedActors/Source/InstancedActors/Private/InstancedActorsCommands.cpp b/Engine/Plugins/Runtime/InstancedActors/Source/InstancedActors/Private/InstancedActorsCommands.cpp
--- a/Engine/Plugins/Runtime/InstancedActors/Source/InstancedActors/Private/InstancedActorsCommands.cpp
+++ b/Engine/Plugins/Runtime/InstancedActors/Source/InstancedActors/Private/InstancedActorsCommands.cpp
@@ -7,6 +7,7 @@
 #include "MassDistanceLODProcessor.h"
 #include "MassStationaryISMSwitcherProcessor.h"
 #include "InstancedActorsVisualizationProcessor.h"
+#include "InstancedActorsData.h"
 
 #ifndef INSTANCEDACTORS_AS_SMARTOBJECTS
 #define INSTANCEDACTORS_AS_SMARTOBJECTS 1
@@ -39,6 +40,32 @@ namespace UE::InstancedActors
 		return DetailedLODTags;
 	}
 
+	void FBulkLODTagsChangeCommandBase::Add(TNotNull<UInstancedActorsData*> InstanceData, EInstancedActorsBulkLOD NewLOD)
+	{
+		Super::Add(InstanceData->Entities);
+		IADataToNotify.Emplace(InstanceData, NewLOD);
+	}
+
+	void FBulkLODTagsChangeCommandBase::Execute(FMassEntityManager& EntityManager) const
+	{
+		Super::Execute(EntityManager);
+
+		for (auto Pair : IADataToNotify)
+		{
+			UInstancedActorsData* IAData = Pair.Get<0>().Get();
+			const EInstancedActorsBulkLOD LOD = Pair.Get<1>();
+
+			// notify the data that it's LOD has changed
+			IAData->OnBulkLODChanged(LOD);
+		}
+	}
+
+	void FBulkLODTagsChangeCommandBase::Reset()
+	{
+		Super::Reset();
+		IADataToNotify.Reset();
+	}
+
 	FEnableDetailedLODCommand::FEnableDetailedLODCommand()
 		: Super(EMassCommandOperationType::Add
 			, GetDetailedLODTags()
diff --git a/Engine/Plugins/Runtime/InstancedActors/Source/InstancedActors/Private/InstancedActorsData.cpp b/Engine/Plugins/Runtime/InstancedActors/Source/InstancedActors/Private/InstancedActorsData.cpp
--- a/Engine/Plugins/Runtime/InstancedActors/Source/InstancedActors/Private/InstancedActorsData.cpp
+++ b/Engine/Plugins/Runtime/InstancedActors/Source/InstancedActors/Private/InstancedActorsData.cpp
@@ -280,6 +280,8 @@ void UInstancedActorsData::Initialize()
 	CreateEntityTemplate(*ExemplarActor);
 
 	bHasEverInitialized = true;
+
+	OnInitialize();
 }
 
 void UInstancedActorsData::Deinitialize()
@@ -292,6 +294,8 @@ void UInstancedActorsData::Deinitialize()
 	}
 
 	CachedSetReplicatedActorRequests.Empty();
+
+	OnDeinitialize();
 }
 
 void UInstancedActorsData::CreateEntityTemplate(const AActor& ExemplarActor)
@@ -415,6 +419,8 @@ void UInstancedActorsData::SpawnEntities()
 	}
 	// we can afford to release the memory, since this array is only ever expected to be used at IAD's boot up
 	CachedSetReplicatedActorRequests.Empty();
+
+	OnSpawnEntities();
 }
 
 void UInstancedActorsData::DespawnEntities()
@@ -510,6 +516,8 @@ void UInstancedActorsData::DespawnEntities()
 	// Reset delta list, if this actor gets recycled on the server we'll get another persistence update restoring the deltas,
 	// if its recycled on the client, the network shadow state is the CDO state so we'll get this replicated again from fresh.
 	InstanceDeltas.Reset(/*bMarkDirty*/ false);
+
+	OnDespawnEntities();
 }
 
 AInstancedActorsManager* UInstancedActorsData::GetManager() const
diff --git a/Engine/Plugins/Runtime/InstancedActors/Source/InstancedActors/Private/InstancedActorsStationaryLODBatchProcessor.cpp b/Engine/Plugins/Runtime/InstancedActors/Source/InstancedActors/Private/InstancedActorsStationaryLODBatchProcessor.cpp
--- a/Engine/Plugins/Runtime/InstancedActors/Source/InstancedActors/Private/InstancedActorsStationaryLODBatchProcessor.cpp
+++ b/Engine/Plugins/Runtime/InstancedActors/Source/InstancedActors/Private/InstancedActorsStationaryLODBatchProcessor.cpp
@@ -343,7 +343,7 @@ void UInstancedActorsStationaryLODBatchProcessor::Execute(FMassEntityManager& En
 						// NOTE: Forcibly updates the mass LOD to off or low when bulk LOD is smaller than Detailed
 						if (InstancedActorsDataTickInfo.BulkLOD == EInstancedActorsBulkLOD::Detailed && InstanceData->CanHydrate())
 						{
-							EntityManager.Defer().PushCommand<UE::InstancedActors::FEnableDetailedLODCommand>(InstanceData->Entities);
+							EntityManager.Defer().PushCommand<UE::InstancedActors::FEnableDetailedLODCommand>(InstanceData, EInstancedActorsBulkLOD::Detailed);
 						}
 						else
 						{
@@ -383,7 +383,7 @@ void UInstancedActorsStationaryLODBatchProcessor::Execute(FMassEntityManager& En
 								});
 
 							// Removes a bunch of tags from all mass entities that belong to an InstancedActorsData, so that we don't spend MassProcessor time on them
-							EntityManager.Defer().PushCommand<UE::InstancedActors::FEnableBatchLODCommand>(InstanceData->Entities);
+							EntityManager.Defer().PushCommand<UE::InstancedActors::FEnableBatchLODCommand>(InstanceData, InstancedActorsDataTickInfo.BulkLOD);
 						}
 					}
 				}
diff --git a/Engine/Plugins/Runtime/InstancedActors/Source/InstancedActors/Public/InstancedActorsCommands.h b/Engine/Plugins/Runtime/InstancedActors/Source/InstancedActors/Public/InstancedActorsCommands.h
--- a/Engine/Plugins/Runtime/InstancedActors/Source/InstancedActors/Public/InstancedActorsCommands.h
+++ b/Engine/Plugins/Runtime/InstancedActors/Source/InstancedActors/Public/InstancedActorsCommands.h
@@ -15,17 +15,29 @@ namespace UE::InstancedActors
 	 */
 	INSTANCEDACTORS_API FMassTagBitSet& GetDetailedLODTags();
 
-	/** Adds GetDetailedLODTags() to an entity, effectively enabling DetailedLOD processing on it */
-	struct FEnableDetailedLODCommand : public FMassCommandChangeTags
+	struct FBulkLODTagsChangeCommandBase : public FMassCommandChangeTags
 	{
+		using FMassCommandChangeTags::FMassCommandChangeTags;
 		using Super = FMassCommandChangeTags;
+		void Add(TNotNull<UInstancedActorsData*> InstanceData, EInstancedActorsBulkLOD NewLOD);
+
+	protected:
+		virtual void Execute(FMassEntityManager& EntityManager) const override;
+		virtual void Reset() override;
+		mutable TArray<TPair<TWeakObjectPtr<UInstancedActorsData>, EInstancedActorsBulkLOD>> IADataToNotify;
+	};
+
+	/** Adds GetDetailedLODTags() to an entity, effectively enabling DetailedLOD processing on it */
+	struct FEnableDetailedLODCommand : public FBulkLODTagsChangeCommandBase 
+	{
+		using Super = FBulkLODTagsChangeCommandBase;
 		FEnableDetailedLODCommand();
 	};
 
 	/** Removes GetDetailedLODTags() from an entity, effectively enabling BatchLOD processing on it */
-	struct FEnableBatchLODCommand : public FMassCommandChangeTags
+	struct FEnableBatchLODCommand : public FBulkLODTagsChangeCommandBase 
 	{
-		using Super = FMassCommandChangeTags;
+		using Super = FBulkLODTagsChangeCommandBase;
 		FEnableBatchLODCommand();
 	};
 }
diff --git a/Engine/Plugins/Runtime/InstancedActors/Source/InstancedActors/Public/InstancedActorsData.h b/Engine/Plugins/Runtime/InstancedActors/Source/InstancedActors/Public/InstancedActorsData.h
--- a/Engine/Plugins/Runtime/InstancedActors/Source/InstancedActors/Public/InstancedActorsData.h
+++ b/Engine/Plugins/Runtime/InstancedActors/Source/InstancedActors/Public/InstancedActorsData.h
@@ -107,7 +107,29 @@ public:
 		SharedSettings = InSharedSettings;
 	}
 
+	// This event virtual is public because it is called by InstancedActorsStationaryLODBatchProcessor
+	UE_API virtual void OnBulkLODChanged(EInstancedActorsBulkLOD NewBulkLOD)
+	{
+	}
+
 protected:
+
+	UE_API virtual void OnInitialize()
+	{
+	}
+
+	UE_API virtual void OnSpawnEntities()
+	{
+	}
+
+	UE_API virtual void OnDespawnEntities()
+	{
+	}
+
+	UE_API virtual void OnDeinitialize()
+	{
+	}
+
 	// Delimiting tag set to group sets of instances
 	UPROPERTY(VisibleAnywhere, Category=InstancedActors)
 	FInstancedActorsTagSet AdditionalTags;
