Update UInstancedActorsStationaryLODBatchProcessor to use ManagerSharedFragment instead of InstancedActorsDataTickInfo. Original change is 44536477 by [REDACTED_HANDLE]

diff --git a/Engine/Plugins/Runtime/InstancedActors/Source/InstancedActors/Private/InstancedActorsStationaryLODBatchProcessor.cpp b/Engine/Plugins/Runtime/InstancedActors/Source/InstancedActors/Private/InstancedActorsStationaryLODBatchProcessor.cpp
--- a/Engine/Plugins/Runtime/InstancedActors/Source/InstancedActors/Private/InstancedActorsStationaryLODBatchProcessor.cpp
+++ b/Engine/Plugins/Runtime/InstancedActors/Source/InstancedActors/Private/InstancedActorsStationaryLODBatchProcessor.cpp
@@ -383,7 +383,7 @@ void UInstancedActorsStationaryLODBatchProcessor::Execute(FMassEntityManager& En
 								});
 
 							// Removes a bunch of tags from all mass entities that belong to an InstancedActorsData, so that we don't spend MassProcessor time on them
-							EntityManager.Defer().PushCommand<UE::InstancedActors::FEnableBatchLODCommand>(InstanceData, InstancedActorsDataTickInfo.BulkLOD);
+							EntityManager.Defer().PushCommand<UE::InstancedActors::FEnableBatchLODCommand>(InstanceData, ManagerSharedFragment.BulkLOD);
 						}
 					}
 				}
