Ran UnrealCodeFixup to move dllstorage to correct places in code

diff --git a/Engine/Plugins/Runtime/InstancedActors/Source/InstancedActors/Public/InstancedActorsData.h b/Engine/Plugins/Runtime/InstancedActors/Source/InstancedActors/Public/InstancedActorsData.h
--- a/Engine/Plugins/Runtime/InstancedActors/Source/InstancedActors/Public/InstancedActorsData.h
+++ b/Engine/Plugins/Runtime/InstancedActors/Source/InstancedActors/Public/InstancedActorsData.h
@@ -121,25 +121,25 @@ public:
 	}
 
 	// This event virtual is public because it is called by InstancedActorsStationaryLODBatchProcessor
-	UE_API virtual void OnBulkLODChanged(EInstancedActorsBulkLOD NewBulkLOD)
+	virtual void OnBulkLODChanged(EInstancedActorsBulkLOD NewBulkLOD)
 	{
 	}
 
 protected:
 
-	UE_API virtual void OnInitialize()
+	virtual void OnInitialize()
 	{
 	}
 
-	UE_API virtual void OnSpawnEntities()
+	virtual void OnSpawnEntities()
 	{
 	}
 
-	UE_API virtual void OnDespawnEntities()
+	virtual void OnDespawnEntities()
 	{
 	}
 
-	UE_API virtual void OnDeinitialize()
+	virtual void OnDeinitialize()
 	{
 	}
 
