Added bInstancedActorsDataDespawnedEntities flag to FInstancedActorsDataSharedFragment
This avoids nulling the shared fragment's InstanceData, which causes issues when trying to retrieve the shared fragment in UInstancedActorsVisualizationTrait::BuildTemplate()

diff --git a/Engine/Plugins/Runtime/InstancedActors/Source/InstancedActors/Private/InstancedActorsData.cpp b/Engine/Plugins/Runtime/InstancedActors/Source/InstancedActors/Private/InstancedActorsData.cpp
--- a/Engine/Plugins/Runtime/InstancedActors/Source/InstancedActors/Private/InstancedActorsData.cpp
+++ b/Engine/Plugins/Runtime/InstancedActors/Source/InstancedActors/Private/InstancedActorsData.cpp
@@ -169,6 +169,13 @@ namespace UE::InstancedActors
 			bCanHydrateLogicEnabled,
 			TEXT("Toggles special handling of instanced marked as non-hydrating via settings (MaxActorDistance ==0)"),
 			ECVF_Default);
+
+		bool bFlagSharedFragmentInsteadOfNullify = true;
+		FAutoConsoleVariableRef CVarFlagSharedFragmentInsteadOfNullify(
+			TEXT("IA.FlagSharedFragmentInsteadOfNullify"),
+			bFlagSharedFragmentInsteadOfNullify,
+			TEXT("Enables using the bInstancedActorsDataDespawnedEntities to mark the IAD shared fragment as deinitialized instead of nullifying the IAD reference."),
+			ECVF_Cheat);
 	} // CVars
 
 	namespace Helpers
@@ -534,7 +541,14 @@ void UInstancedActorsData::DespawnEntities()
 		// decrement the stats before resetting
 		Manager.UpdateInstanceStats(NumInstances, AsSharedManagerFragment->BulkLOD, /*Increment=*/false);
 
-		AsSharedManagerFragment->InstanceData = nullptr;
+		if (UE::InstancedActors::CVars::bFlagSharedFragmentInsteadOfNullify)
+		{
+			AsSharedManagerFragment->bInstancedActorsDataDespawnedEntities = true;
+		}
+		else
+		{
+			AsSharedManagerFragment->InstanceData = nullptr;
+		}
 		AsSharedManagerFragment->BulkLOD = EInstancedActorsBulkLOD::MAX;
 	}
 
diff --git a/Engine/Plugins/Runtime/InstancedActors/Source/InstancedActors/Private/InstancedActorsStationaryLODBatchProcessor.cpp b/Engine/Plugins/Runtime/InstancedActors/Source/InstancedActors/Private/InstancedActorsStationaryLODBatchProcessor.cpp
--- a/Engine/Plugins/Runtime/InstancedActors/Source/InstancedActors/Private/InstancedActorsStationaryLODBatchProcessor.cpp
+++ b/Engine/Plugins/Runtime/InstancedActors/Source/InstancedActors/Private/InstancedActorsStationaryLODBatchProcessor.cpp
@@ -231,7 +231,17 @@ void UInstancedActorsStationaryLODBatchProcessor::Execute(FMassEntityManager& En
 			(FInstancedActorsDataSharedFragment& ManagerSharedFragment) -> double
 			{
 				double NextTickTime = CurrentTime + DelayPerBulkLOD[(int)EInstancedActorsBulkLOD::Off];
-				if (UInstancedActorsData* InstanceData = ManagerSharedFragment.InstanceData.Get())
+
+				UInstancedActorsData* InstanceData = ManagerSharedFragment.InstanceData.Get();
+				if (UE::InstancedActors::CVars::bFlagSharedFragmentInsteadOfNullify)
+				{
+					if (ManagerSharedFragment.bInstancedActorsDataDespawnedEntities)
+					{
+						InstanceData = nullptr;
+					}
+				}
+
+				if (InstanceData != nullptr)
 				{
 					const FInstancedActorsSettings& Settings = InstanceData->GetSettings<const FInstancedActorsSettings>();
 					
diff --git a/Engine/Plugins/Runtime/InstancedActors/Source/InstancedActors/Private/InstancedActorsVisualizationTrait.cpp b/Engine/Plugins/Runtime/InstancedActors/Source/InstancedActors/Private/InstancedActorsVisualizationTrait.cpp
--- a/Engine/Plugins/Runtime/InstancedActors/Source/InstancedActors/Private/InstancedActorsVisualizationTrait.cpp
+++ b/Engine/Plugins/Runtime/InstancedActors/Source/InstancedActors/Private/InstancedActorsVisualizationTrait.cpp
@@ -102,11 +102,13 @@ void UInstancedActorsVisualizationTrait::BuildTemplate(FMassEntityTemplateBuildC
 	{
 		// this can happen when we unload data and then stream it back again - we end up with the same path object, but different pointer.
 		// The old instance should be garbage now
-		if (AsShared->InstanceData != InstanceData)
+		if (AsShared->bInstancedActorsDataDespawnedEntities || AsShared->InstanceData != InstanceData)
 		{
-			ensure(AsShared->InstanceData.IsValid() == false);
+			ensure(AsShared->bInstancedActorsDataDespawnedEntities || AsShared->InstanceData.IsValid() == false);
 			AsShared->InstanceData = InstanceData;
 		}
+
+		AsShared->bInstancedActorsDataDespawnedEntities = false;
 		
 		// we also need to make sure the bulk LOD is reset since the shared fragment can survive the death of the original 
 		// InstanceData while preserving the "runtime" value, which will mess up newly spawned entities.
diff --git a/Engine/Plugins/Runtime/InstancedActors/Source/InstancedActors/Public/InstancedActorsData.h b/Engine/Plugins/Runtime/InstancedActors/Source/InstancedActors/Public/InstancedActorsData.h
--- a/Engine/Plugins/Runtime/InstancedActors/Source/InstancedActors/Public/InstancedActorsData.h
+++ b/Engine/Plugins/Runtime/InstancedActors/Source/InstancedActors/Public/InstancedActorsData.h
@@ -14,6 +14,14 @@
 class AInstancedActorsManager;
 struct FInstancedActorsSettings;
 
+namespace UE::InstancedActors
+{
+	namespace CVars
+	{
+		extern bool bFlagSharedFragmentInsteadOfNullify;
+	}
+}
+
 namespace UE::InstancedActors
 {
 struct FExemplarActorData
diff --git a/Engine/Plugins/Runtime/InstancedActors/Source/InstancedActors/Public/InstancedActorsTypes.h b/Engine/Plugins/Runtime/InstancedActors/Source/InstancedActors/Public/InstancedActorsTypes.h
--- a/Engine/Plugins/Runtime/InstancedActors/Source/InstancedActors/Public/InstancedActorsTypes.h
+++ b/Engine/Plugins/Runtime/InstancedActors/Source/InstancedActors/Public/InstancedActorsTypes.h
@@ -333,6 +333,9 @@ struct FInstancedActorsDataSharedFragment : public FMassSharedFragment
 
 	EInstancedActorsBulkLOD BulkLOD = EInstancedActorsBulkLOD::MAX;
 
+	// Set to true when the IAD has despawned its entities
+	bool bInstancedActorsDataDespawnedEntities = false;
+
 	double LastTickTime = 0.0;
 };
 
