FInstancedActorsInstanceIndex now uses its own INSTANCED_INDEX_NONE value based on MAX_uint16 instead of using INDEX_NONE (-1). This prevents the index from wrapping and causing unexpected behaviors

#FailedValidations Preflight

diff --git a/Engine/Plugins/Runtime/InstancedActors/Source/InstancedActors/Private/InstancedActorsData.cpp b/Engine/Plugins/Runtime/InstancedActors/Source/InstancedActors/Private/InstancedActorsData.cpp
--- a/Engine/Plugins/Runtime/InstancedActors/Source/InstancedActors/Private/InstancedActorsData.cpp
+++ b/Engine/Plugins/Runtime/InstancedActors/Source/InstancedActors/Private/InstancedActorsData.cpp
@@ -647,7 +647,7 @@ FInstancedActorsInstanceHandle UInstancedActorsData::AddInstance(const FTransfor
 		}
 	}
 	// Add a new instance
-	else if (!ensureMsgf(InstanceTransforms.Num() < (MAX_uint16 - 1), TEXT("Instance Actors do not support more than 65534 instances per cell. Make the cells smaller or lower the density of actors."))) // TODO: make this a constant in InstancedActorsIndex
+	else if (!ensureMsgf(InstanceTransforms.Num() < FInstancedActorsInstanceIndex::INSTANCED_INDEX_MAX, TEXT("Instance Actors do not support more than 65534 instances per cell. Make the cells smaller or lower the density of actors.")))
 	{
 		return FInstancedActorsInstanceHandle();
 	}
diff --git a/Engine/Plugins/Runtime/InstancedActors/Source/InstancedActors/Public/InstancedActorsIndex.h b/Engine/Plugins/Runtime/InstancedActors/Source/InstancedActors/Public/InstancedActorsIndex.h
--- a/Engine/Plugins/Runtime/InstancedActors/Source/InstancedActors/Public/InstancedActorsIndex.h
+++ b/Engine/Plugins/Runtime/InstancedActors/Source/InstancedActors/Public/InstancedActorsIndex.h
@@ -18,11 +18,15 @@ USTRUCT()
 struct FInstancedActorsInstanceIndex
 {
 	GENERATED_BODY()
+
+	static constexpr uint16 INSTANCED_INDEX_NONE = MAX_uint16;
+	static constexpr uint16 INSTANCED_INDEX_MAX = MAX_uint16 - 1;
 	
 	FInstancedActorsInstanceIndex() = default;
-	explicit FInstancedActorsInstanceIndex(const int32 InIndex) : Index((uint16)InIndex) 
+	explicit FInstancedActorsInstanceIndex(const int32 InIndex)
 	{ 
-		check((InIndex >= INDEX_NONE) && (InIndex < (MAX_uint16 - 1))); // -1 for INDEX_NONE
+		check(InIndex >= 0 && InIndex < INSTANCED_INDEX_NONE);
+		Index = InIndex;
 	}
 
 	friend FArchive& operator<<(FArchive& Ar, FInstancedActorsInstanceIndex& InstanceIndex)
@@ -41,7 +45,7 @@ struct FInstancedActorsInstanceIndex
 		return uint32(InstanceIndex.Index);
 	}
 
-	bool IsValid() const { return Index != INDEX_NONE; }
+	bool IsValid() const { return Index != INSTANCED_INDEX_NONE; }
 
 	/** Returns a string suitable for debug logging to identify this instance. */
 	UE_API FString GetDebugName() const;
@@ -52,6 +56,8 @@ struct FInstancedActorsInstanceIndex
 
 	static constexpr inline int32 BuildCompositeIndex(const uint16 InstanceDataID, const int32 InstanceIndex)
 	{
+		check(InstanceIndex >= 0 && InstanceIndex <= INSTANCED_INDEX_MAX);
+
 		const uint32 HighBits = InstanceDataID;
 		const uint32 LowBits = InstanceIndex;
 		check(LowBits <= MAX_uint16);
@@ -71,7 +77,7 @@ struct FInstancedActorsInstanceIndex
 private:
 	/** Stable(consistent between client and server) instance index into UInstancedActorsData */
 	UPROPERTY()
-	uint16 Index = INDEX_NONE;
+	uint16 Index = INSTANCED_INDEX_NONE;
 
 	static constexpr int32 InstanceIndexBits = 16;
 	static constexpr int32 InstanceIndexMask = (1 << InstanceIndexBits) - 1;
