Added new flag signaling subsystem switching when despawning entities to address timing issue with latent subsystem creation.
IAMs will get be switched over when the new game-mode specific subsystem is created, which will clear the IAD pointer or flag that entity despawn has occured on the shared fragment, preventing the batch LOD processor from doing its checks. This change prevents that in this situation

diff --git a/Engine/Plugins/Runtime/InstancedActors/Source/InstancedActors/Private/InstancedActorsData.cpp b/Engine/Plugins/Runtime/InstancedActors/Source/InstancedActors/Private/InstancedActorsData.cpp
--- a/Engine/Plugins/Runtime/InstancedActors/Source/InstancedActors/Private/InstancedActorsData.cpp
+++ b/Engine/Plugins/Runtime/InstancedActors/Source/InstancedActors/Private/InstancedActorsData.cpp
@@ -441,7 +441,7 @@ void UInstancedActorsData::SpawnEntities()
 	OnSpawnEntities();
 }
 
-void UInstancedActorsData::DespawnEntities()
+void UInstancedActorsData::DespawnEntities(bool bIsSwitchingSubsystems /*= false*/)
 {
 	UWorld* World = GetWorld();
 	check(World);
@@ -539,13 +539,16 @@ void UInstancedActorsData::DespawnEntities()
 		// decrement the stats before resetting
 		Manager.UpdateInstanceStats(NumInstances, AsSharedManagerFragment->BulkLOD, /*Increment=*/false);
 
-		if (UE::InstancedActors::CVars::bFlagSharedFragmentInsteadOfNullify)
+		if (!bIsSwitchingSubsystems)
 		{
-			AsSharedManagerFragment->bInstancedActorsDataDespawnedEntities = true;
-		}
-		else
-		{
-			AsSharedManagerFragment->InstanceData = nullptr;
+			if (UE::InstancedActors::CVars::bFlagSharedFragmentInsteadOfNullify)
+			{
+				AsSharedManagerFragment->bInstancedActorsDataDespawnedEntities = true;
+			}
+			else
+			{
+				AsSharedManagerFragment->InstanceData = nullptr;
+			}
 		}
 		AsSharedManagerFragment->BulkLOD = EInstancedActorsBulkLOD::MAX;
 	}
diff --git a/Engine/Plugins/Runtime/InstancedActors/Source/InstancedActors/Private/InstancedActorsManager.cpp b/Engine/Plugins/Runtime/InstancedActors/Source/InstancedActors/Private/InstancedActorsManager.cpp
--- a/Engine/Plugins/Runtime/InstancedActors/Source/InstancedActors/Private/InstancedActorsManager.cpp
+++ b/Engine/Plugins/Runtime/InstancedActors/Source/InstancedActors/Private/InstancedActorsManager.cpp
@@ -286,7 +286,7 @@ void AInstancedActorsManager::OnAddedToSubsystem(UInstancedActorsSubsystem& InIn
 	if (IsValid(InstancedActorSubsystem) && (InstancedActorSubsystem != &InInstancedActorSubsystem))
 	{
 		// we need to unregister from the previous subsystem first
-		DespawnAllEntities();
+		DespawnAllEntities(/*bIsSwitchingSubsystems*/ true);
 		// and let registered modifiers go
 		RemoveAllModifierVolumes();
 		
@@ -908,7 +908,7 @@ void AInstancedActorsManager::RuntimeRemoveAllInstances()
 	}
 }
 
-void AInstancedActorsManager::DespawnAllEntities()
+void AInstancedActorsManager::DespawnAllEntities(bool bIsSwitchingSubsystems /*= false*/)
 {
 	// DespawnEntities for all PerActorClassInstanceData
 	if (HasSpawnedEntities())
@@ -916,7 +916,7 @@ void AInstancedActorsManager::DespawnAllEntities()
 		for (TObjectPtr<UInstancedActorsData> InstanceData : PerActorClassInstanceData)
 		{
 			// Reconstruct instance data from Mass then destroy all Mass entities
-			InstanceData->DespawnEntities();
+			InstanceData->DespawnEntities(bIsSwitchingSubsystems);
 		}
 		bHasSpawnedEntities = false;
 	}
diff --git a/Engine/Plugins/Runtime/InstancedActors/Source/InstancedActors/Public/InstancedActorsData.h b/Engine/Plugins/Runtime/InstancedActors/Source/InstancedActors/Public/InstancedActorsData.h
--- a/Engine/Plugins/Runtime/InstancedActors/Source/InstancedActors/Public/InstancedActorsData.h
+++ b/Engine/Plugins/Runtime/InstancedActors/Source/InstancedActors/Public/InstancedActorsData.h
@@ -66,7 +66,7 @@ public:
 
 	// Called early in AInstancedActorsManager::EndPlay to reconstruct cooked data state from runtime Mass entities as best we can,
 	// then despawn all Mass entities and reset any other runtime instance data
-	UE_API void DespawnEntities();
+	UE_API void DespawnEntities(bool bIsSwitchingSubsystems = false);
 
 	// Called from the parent's AInstancedActorsManager::EndPlay. Can release the entity template and exemplar Actor from memory 
 	UE_API void Deinitialize();
diff --git a/Engine/Plugins/Runtime/InstancedActors/Source/InstancedActors/Public/InstancedActorsManager.h b/Engine/Plugins/Runtime/InstancedActors/Source/InstancedActors/Public/InstancedActorsManager.h
--- a/Engine/Plugins/Runtime/InstancedActors/Source/InstancedActors/Public/InstancedActorsManager.h
+++ b/Engine/Plugins/Runtime/InstancedActors/Source/InstancedActors/Public/InstancedActorsManager.h
@@ -292,7 +292,7 @@ protected:
 	UE_API void SerializeInstancePersistenceData(FStructuredArchive::FRecord Record, UInstancedActorsData* InstanceData, int64 TimeDelta) const;
 
 	/** Despawns all entities spawned by individual UInstancedActorsData instances. */
-	UE_API virtual void DespawnAllEntities();
+	UE_API virtual void DespawnAllEntities(bool bIsSwitchingSubsystems = false);
 
 	/** Attempts to run any 'pending' modifiers in ModifierVolumes where are appropriate to run given HasSpawnedEntities
 	 * Called in BeginPlay prior to, and then again after SpawnEntities. Also called in AddModifierVolume.
