[InstancedActors] Destroy/release the UInstancedActorsData's entity template when calling RemoveVisualization(). This is a workaround fix - the issue is that the entity template/UInstancedActorsVisualizationTrait treats the FMassRepresentationFragment::StaticMeshDescHandle as invariant.

Removing a visualization means that an entity's FMassRepresentationFragment::StaticMeshDescHandle becomes silently invalid (the index could point to blank ISM data, or that index can be re-used by another UInstancedActorsData populating it with entirely different ISM data). The FMassRepresentationFragment data is setup via UInstancedActorsVisualizationTrait, but any future calls to UInstancedActorsData::CreateEntityTemplate() will not allow a recomputation of this trait because it already exists. By releasing the entity template, we essentially force UInstancedActorsVisualizationTrait::BuildTemplate() to be called again, thus assigning FMassRepresentationFragment::StaticMeshDescHandle to that matching the proper visualization.

In practice, we would see this happen when the parent IAM is switching subsystems, which then forces UInstancedActorsData::DespawnEntities() to be called, and then RemoveAllVisualizations().

diff --git a/Engine/Plugins/Runtime/InstancedActors/Source/InstancedActors/Private/InstancedActorsData.cpp b/Engine/Plugins/Runtime/InstancedActors/Source/InstancedActors/Private/InstancedActorsData.cpp
--- a/Engine/Plugins/Runtime/InstancedActors/Source/InstancedActors/Private/InstancedActorsData.cpp
+++ b/Engine/Plugins/Runtime/InstancedActors/Source/InstancedActors/Private/InstancedActorsData.cpp
@@ -174,6 +174,13 @@ namespace UE::InstancedActors
 			bFlagSharedFragmentInsteadOfNullify,
 			TEXT("Enables using the bInstancedActorsDataDespawnedEntities to mark the IAD shared fragment as deinitialized instead of nullifying the IAD reference."),
 			ECVF_Cheat);
+
+		bool bReleaseEntityTemplateWhenRemovingVisualization = true;
+		FAutoConsoleVariableRef CVarReleaseEntityTemplateWhenRemovingVisualization(
+			TEXT("IA.ReleaseEntityTemplateWhenRemovingVisualization"),
+			bReleaseEntityTemplateWhenRemovingVisualization,
+			TEXT("If enabled (default) instanced actors will release created entity templates when a visualization is removed. This ensures recreating the entities using the newly created visualization"),
+			ECVF_Default);
 	} // CVars
 
 	namespace Helpers
@@ -1695,7 +1702,21 @@ void UInstancedActorsData::RemoveVisualization(uint8 VisualizationIndex)
 	// Reset RemovedVisualization and set it's allocation flag to false, allowing for reuse in subsequent AddVisualization calls
 	RemovedVisualization = FInstancedActorsVisualizationInfo();
 	check(InstanceVisualizationAllocationFlags.IsValidIndex(VisualizationIndex));
-	InstanceVisualizationAllocationFlags[VisualizationIndex] = false;
+	InstanceVisualizationAllocationFlags[VisualizationIndex] = false;	
+
+	// Destroying visualizations can leave the FMassRepresentationFragment::StaticMeshDescHandle pointing to stale data. So we need to rebuild the UInstancedActorsVisualizationTrait
+	// This is a workaround fix - we are relying on the fact that RemoveVisualization() is only being called when DespawnEntities() is called. This means that either:
+	// (1) The parent AInstancedActorsManager::EndPlay() is being called
+	// (2) The parent AInstancedActorsManager::OnAddedToSubsystem() is being called with a request to change subsystems -
+	//     so we despawn all entities, remove all visualizations, then respawn all entities
+	// In that second case, removing the visualizations means it invalidates the template's FMassRepresentationFragment::StaticMeshDescHandle.
+	// If we don't force the template to be rebuilt, when this IAD respawns the entities and calls CreateEntityTemplate(), it would just fallback to the previously created template
+	// The issue is that the entity template/UInstancedActorsVisualizationTrait treats the FMassRepresentationFragment::StaticMeshDescHandle as invariant
+	// A better long-term solve would be to set FMassRepresentationFragment::StaticMeshDescHandle post-entity creation. And possibly error-checking the context of visualizations being removed
+	if (UE::InstancedActors::CVars::bReleaseEntityTemplateWhenRemovingVisualization)
+	{
+		ReleaseEntityTemplate();
+	}
 }
 
 void UInstancedActorsData::RemoveAllVisualizations()
