InstancedActorsData's runtime instance removal is now defered if the MassEntityManager is processing.
This fixes a bug where entities could sometimes be destroyed while being iterated on.

diff --git a/Engine/Plugins/Runtime/InstancedActors/Source/InstancedActors/Private/InstancedActorsData.cpp b/Engine/Plugins/Runtime/InstancedActors/Source/InstancedActors/Private/InstancedActorsData.cpp
--- a/Engine/Plugins/Runtime/InstancedActors/Source/InstancedActors/Private/InstancedActorsData.cpp
+++ b/Engine/Plugins/Runtime/InstancedActors/Source/InstancedActors/Private/InstancedActorsData.cpp
@@ -1237,7 +1237,17 @@ void UInstancedActorsData::RuntimeRemoveInstances(TConstArrayView<FInstancedActo
 
 		// Destroy entities. This will also trigger actor destruction for any spawned actors
 		// for these entities.
-		MassEntityManager.BatchDestroyEntityChunks(EntityCollectionsToDestroy);
+		if (!EntityCollectionsToDestroy.IsEmpty())
+		{
+			if (MassEntityManager.IsProcessing())
+			{
+				MassEntityManager.Defer().DestroyEntities(EntitiesToDestroy);
+			}
+			else
+			{
+				MassEntityManager.BatchDestroyEntityChunks(EntityCollectionsToDestroy);
+			}
+		}
 	}
 	// Pre-empt entity spawning and simply invalidate InstanceTransform entries, preventing them from spawning later
 	else
