ISM PSO Missing Precaching fix for MDC category "LightWeightInstance"

diff --git a/Engine/Plugins/Runtime/InstancedActors/Source/InstancedActors/Private/InstancedActorsManager.cpp b/Engine/Plugins/Runtime/InstancedActors/Source/InstancedActors/Private/InstancedActorsManager.cpp
--- a/Engine/Plugins/Runtime/InstancedActors/Source/InstancedActors/Private/InstancedActorsManager.cpp
+++ b/Engine/Plugins/Runtime/InstancedActors/Source/InstancedActors/Private/InstancedActorsManager.cpp
@@ -1627,6 +1627,8 @@ void AInstancedActorsManager::CreateISMComponents(const FInstancedActorsVisualiz
 			SetUpEditorPreviewISMComponent(ISMComponent);
 		}
 
+		ISMComponent->PrecachePSOs();
+
 		ISMComponent->RegisterComponent();
 
 		OutComponents.Add(ISMComponent);
