InstancedActorsData's runtime instance removal now defers entity destruction by default. This behavior can be changed with IA.DeferRuntimeRemoveInstancesEntityDestruction.

diff --git a/Engine/Plugins/Runtime/InstancedActors/Source/InstancedActors/Private/InstancedActorsData.cpp b/Engine/Plugins/Runtime/InstancedActors/Source/InstancedActors/Private/InstancedActorsData.cpp
--- a/Engine/Plugins/Runtime/InstancedActors/Source/InstancedActors/Private/InstancedActorsData.cpp
+++ b/Engine/Plugins/Runtime/InstancedActors/Source/InstancedActors/Private/InstancedActorsData.cpp
@@ -181,6 +181,13 @@ namespace UE::InstancedActors
 			bReleaseEntityTemplateWhenRemovingVisualization,
 			TEXT("If enabled (default) instanced actors will release created entity templates when a visualization is removed. This ensures recreating the entities using the newly created visualization"),
 			ECVF_Default);
+
+		bool bDeferRuntimeRemoveInstancesEntityDestruction = true;
+		FAutoConsoleVariableRef CVarbDeferRuntimeRemoveInstancesEntityDestruction(
+			TEXT("IA.DeferRuntimeRemoveInstancesEntityDestruction"),
+			bDeferRuntimeRemoveInstancesEntityDestruction,
+			TEXT("If enabled, instanced actors entity destruction is deferred instead of being synchronous when removing the instance at runtime."),
+			ECVF_Default);
 	} // CVars
 
 	namespace Helpers
@@ -1239,7 +1246,7 @@ void UInstancedActorsData::RuntimeRemoveInstances(TConstArrayView<FInstancedActo
 		// for these entities.
 		if (!EntityCollectionsToDestroy.IsEmpty())
 		{
-			if (MassEntityManager.IsProcessing())
+			if (UE::InstancedActors::CVars::bDeferRuntimeRemoveInstancesEntityDestruction)
 			{
 				MassEntityManager.Defer().DestroyEntities(EntitiesToDestroy);
 			}
