Converting loop index from uint16 to int32 to avoid type promotion issues

diff --git a/Engine/Plugins/Runtime/InstancedActors/Source/InstancedActors/Private/InstancedActorsManager.cpp b/Engine/Plugins/Runtime/InstancedActors/Source/InstancedActors/Private/InstancedActorsManager.cpp
--- a/Engine/Plugins/Runtime/InstancedActors/Source/InstancedActors/Private/InstancedActorsManager.cpp
+++ b/Engine/Plugins/Runtime/InstancedActors/Source/InstancedActors/Private/InstancedActorsManager.cpp
@@ -563,12 +563,12 @@ void AInstancedActorsManager::SerializeInstancePersistenceData(FStructuredArchiv
 	// @todo: In hindsight we could / should just intrinisically have separate arrays / 'delta lists' of Destroyed Instances and Lifecycle Changes.
 
 	// Destroyed instance indices
-	int32 NumDestroyedInstances = InstanceData ? InstanceData->InstanceDeltas.GetNumDestroyedInstanceDeltas() : INDEX_NONE;
+	int32 NumDestroyedInstances = InstanceData ? InstanceData->InstanceDeltas.GetNumDestroyedInstanceDeltas() : 0;
 	FStructuredArchive::FArray DestroyedInstancesArray = Record.EnterArray(TEXT("DestroyedInstances"), NumDestroyedInstances);
 
 	if (UnderlyingArchive.IsLoading())
 	{
-		for (uint16 ArrayIndex = 0; ArrayIndex < NumDestroyedInstances; ++ArrayIndex)
+		for (int32 ArrayIndex = 0; ArrayIndex < NumDestroyedInstances; ++ArrayIndex)
 		{
 			if (!ensureMsgf(InstanceData && ArrayIndex < InstanceData->GetNumInstances(), TEXT("Attempting to destroy more instances (%d) than actually exist (%d). Aborting corrupted persistence archive read. Persistence data may be lost as a result."), NumDestroyedInstances, InstanceData->GetNumInstances()))
 			{
@@ -594,7 +594,7 @@ void AInstancedActorsManager::SerializeInstancePersistenceData(FStructuredArchiv
 	{
 		check(InstanceData);
 
-		uint16 NumDestroyedInstancesWritten = 0;
+		int32 NumDestroyedInstancesWritten = 0;
 		const TArray<FInstancedActorsDelta>& InstanceDeltas = InstanceData->InstanceDeltas.GetInstanceDeltas();
 		for (const FInstancedActorsDelta& Delta : InstanceDeltas)
 		{
