Properly guarding against InstanceData being null
This fixes a segfault caused by the ensureMsgf log trying to access a null InstanceData.

diff --git a/Engine/Plugins/Runtime/InstancedActors/Source/InstancedActors/Private/InstancedActorsManager.cpp b/Engine/Plugins/Runtime/InstancedActors/Source/InstancedActors/Private/InstancedActorsManager.cpp
--- a/Engine/Plugins/Runtime/InstancedActors/Source/InstancedActors/Private/InstancedActorsManager.cpp
+++ b/Engine/Plugins/Runtime/InstancedActors/Source/InstancedActors/Private/InstancedActorsManager.cpp
@@ -570,7 +570,10 @@ void AInstancedActorsManager::SerializeInstancePersistenceData(FStructuredArchiv
 	{
 		for (int32 ArrayIndex = 0; ArrayIndex < NumDestroyedInstances; ++ArrayIndex)
 		{
-			if (!ensureMsgf(InstanceData && ArrayIndex < InstanceData->GetNumInstances(), TEXT("Attempting to destroy more instances (%d) than actually exist (%d). Aborting corrupted persistence archive read. Persistence data may be lost as a result."), NumDestroyedInstances, InstanceData->GetNumInstances()))
+			const int32 NumInstances = InstanceData ? InstanceData->GetNumInstances() : 0;
+			if (!ensureMsgf(InstanceData && ArrayIndex < NumInstances,
+				TEXT("Attempting to destroy more instances (%d) than actually exist (%d). Aborting corrupted persistence archive read. Persistence data may be lost as a result."),
+				NumDestroyedInstances, NumInstances))
 			{
 				UnderlyingArchive.SetError();
 				return;
