InstancedActors: buffer runtime entity removal and guard Mass visualization processors with a pending-destruction flag to prevent processors from unnecessarily considering those entities.

UInstancedActorsSubsystem now tracks InstancedActorsData with entities pending destruction.
Added a bIsPendingDestruction flag in FMassRepresentationFragment. Downstream processors interacting with the fragment (such as UMassStationaryISMSwitcherProcessor) can now use it to exit early when considering those entities.
This fixes a race where entities pending destruction could re-add an ISM instance and create physics collider because UMassStationaryISMSwitcherProcessor processed them between actor destruction and entity removal.

diff --git a/Engine/Plugins/Runtime/InstancedActors/Source/InstancedActors/Private/InstancedActorsData.cpp b/Engine/Plugins/Runtime/InstancedActors/Source/InstancedActors/Private/InstancedActorsData.cpp
--- a/Engine/Plugins/Runtime/InstancedActors/Source/InstancedActors/Private/InstancedActorsData.cpp
+++ b/Engine/Plugins/Runtime/InstancedActors/Source/InstancedActors/Private/InstancedActorsData.cpp
@@ -825,6 +825,12 @@ FString UInstancedActorsData::GetDebugName(bool bCompact) const
 	}
 }
 
+void UInstancedActorsData::PopAllEntitiesPendingDestruction(TArray<FMassEntityHandle>& InOutEntities)
+{
+    FScopeLock Lock(&PendingDestructionCS);
+    InOutEntities.Append(MoveTemp(EntitiesPendingDestruction));
+}
+
 void UInstancedActorsData::PostDuplicate(EDuplicateMode::Type DuplicateMode)
 {
 	Super::PostDuplicate(DuplicateMode);
@@ -1239,22 +1245,24 @@ void UInstancedActorsData::RuntimeRemoveInstances(TConstArrayView<FInstancedActo
 				}
 			}
 		}
-		TArray<FMassArchetypeEntityCollection> EntityCollectionsToDestroy;
-		UE::Mass::Utils::CreateEntityCollections(MassEntityManager, EntitiesToDestroy, FMassArchetypeEntityCollection::NoDuplicates, EntityCollectionsToDestroy);
 
 		// Destroy entities. This will also trigger actor destruction for any spawned actors
 		// for these entities.
-		if (!EntityCollectionsToDestroy.IsEmpty())
+		if (UE::InstancedActors::CVars::bDeferRuntimeRemoveInstancesEntityDestruction)
 		{
-			if (UE::InstancedActors::CVars::bDeferRuntimeRemoveInstancesEntityDestruction)
+			QueueEntitiesForDestruction(MoveTemp(EntitiesToDestroy));
+			if (UInstancedActorsSubsystem* IASubsystem = GetManagerChecked().GetInstancedActorSubsystem())
 			{
-				MassEntityManager.Defer().DestroyEntities(EntitiesToDestroy);
-			}
-			else
-			{
-				MassEntityManager.BatchDestroyEntityChunks(EntityCollectionsToDestroy);
+				IASubsystem->MarkIADDirtyWithEntitiesPendingDestruction(*this);
 			}
 		}
+		else
+		{
+			TArray<FMassArchetypeEntityCollection> EntityCollectionsToDestroy;
+			UE::Mass::Utils::CreateEntityCollections(MassEntityManager, EntitiesToDestroy, FMassArchetypeEntityCollection::NoDuplicates, EntityCollectionsToDestroy);
+
+			MassEntityManager.BatchDestroyEntityChunks(EntityCollectionsToDestroy);
+		}
 	}
 	// Pre-empt entity spawning and simply invalidate InstanceTransform entries, preventing them from spawning later
 	else
@@ -1553,6 +1561,20 @@ void UInstancedActorsData::InitializeVisualization(uint8 AllocatedVisualizationI
 	}
 }
 
+void UInstancedActorsData::QueueEntitiesForDestruction(TArray<FMassEntityHandle>&& EntitiesToDestroy)
+{
+	if (EntitiesToDestroy.IsEmpty())
+	{
+		return;
+	}
+
+	if (HasSpawnedEntities())
+	{
+		FScopeLock Lock(&PendingDestructionCS);
+		EntitiesPendingDestruction.Append(MoveTemp(EntitiesToDestroy));
+	}
+}
+
 uint8 UInstancedActorsData::AddVisualization(FInstancedActorsVisualizationDesc& InOutVisualizationDesc)
 {
 	// Reuse free or create new InstanceVisualizations entry
diff --git a/Engine/Plugins/Runtime/InstancedActors/Source/InstancedActors/Private/InstancedActorsRuntimeRemovalProcessor.cpp b/Engine/Plugins/Runtime/InstancedActors/Source/InstancedActors/Private/InstancedActorsRuntimeRemovalProcessor.cpp
new file mode 100644
--- /dev/null
+++ b/Engine/Plugins/Runtime/InstancedActors/Source/InstancedActors/Private/InstancedActorsRuntimeRemovalProcessor.cpp
@@ -0,0 +1,90 @@
+// Copyright Epic Games, Inc. All Rights Reserved.
+
+#include "InstancedActorsRuntimeRemovalProcessor.h"
+#include "InstancedActorsData.h"
+#include "InstancedActorsStationaryLODBatchProcessor.h"
+#include "InstancedActorsSubsystem.h"
+#include "MassEntityView.h"
+#include "MassExecutionContext.h"
+#include "MassRepresentationFragments.h"
+#include "MassStationaryISMSwitcherProcessor.h"
+#include "InstancedActorsSettings.h"
+
+#include UE_INLINE_GENERATED_CPP_BY_NAME(InstancedActorsRuntimeRemovalProcessor)
+
+UInstancedActorsRuntimeRemovalProcessor::UInstancedActorsRuntimeRemovalProcessor()
+: EntityQuery(*this)
+{
+	ExecutionOrder.ExecuteInGroup = UE::Mass::ProcessorGroupNames::Representation;
+	bAutoRegisterWithProcessingPhases = true;
+	ExecutionFlags = static_cast<int32>(EProcessorExecutionFlags::AllNetModes);
+
+	ExecutionOrder.ExecuteBefore.Add(UInstancedActorsStationaryLODBatchProcessor::StaticClass()->GetFName());
+	ExecutionOrder.ExecuteBefore.Add(UMassStationaryISMSwitcherProcessor::StaticClass()->GetFName());
+}
+
+void UInstancedActorsRuntimeRemovalProcessor::ConfigureQueries(const TSharedRef<FMassEntityManager>& EntityManager)
+{
+	// Not directly used, but kept for the dependency solver
+	EntityQuery.AddRequirement<FMassRepresentationFragment>(EMassFragmentAccess::ReadWrite);
+
+	if (HasAnyFlags(RF_ClassDefaultObject) == false)
+	{
+		if (TSubclassOf<UInstancedActorsSubsystem> InstancedActorsSubsystemClass = GET_INSTANCEDACTORS_CONFIG_VALUE(GetInstancedActorsSubsystemClass()))
+		{
+			ProcessorRequirements.AddSubsystemRequirement(InstancedActorsSubsystemClass, EMassFragmentAccess::ReadWrite, EntityManager);
+		}
+		else
+		{
+			UE_LOG(LogInstancedActors, Error, TEXT("Misconfigured UInstancedActorsSubsystem subclass"));
+		}
+	}
+}
+
+void UInstancedActorsRuntimeRemovalProcessor::Execute(FMassEntityManager& EntityManager, FMassExecutionContext& Context)
+{
+	TSubclassOf<UInstancedActorsSubsystem> InstancedActorsSubsystemClass = GET_INSTANCEDACTORS_CONFIG_VALUE(GetInstancedActorsSubsystemClass());
+	if (!ensureMsgf(InstancedActorsSubsystemClass, TEXT("Misconfigured UInstancedActorsSubsystem subclass")))
+	{
+		return;
+	}
+	//@todo: revisit this temporary fix. This was done because UGameFeatureAction_ConfigureInstancedActors::OnGameFeatureActivating() may happen after UInstancedActorsRuntimeRemovalProcessor::ConfigureQueries()
+	//UInstancedActorsSubsystem* InstancedActorSubsystem = Context.GetMutableSubsystem<UInstancedActorsSubsystem>(InstancedActorsSubsystemClass);
+	UInstancedActorsSubsystem* InstancedActorsSubsystem = UE::InstancedActors::Utils::GetInstancedActorsSubsystem(*Context.GetWorld());
+	if (!ensureMsgf(InstancedActorsSubsystem, TEXT("UInstancedActorsSubsystem is missing, this is unexpected")))
+	{
+		return;
+	}
+
+	if (!InstancedActorsSubsystem->ContainsDirtyIADWithEntitiesPendingDestruction())
+	{
+		return;
+	}
+
+	TArray<TWeakObjectPtr<UInstancedActorsData>> DirtyIADs;
+	TArray<FMassEntityHandle> EntitiesPendingDestruction;
+	InstancedActorsSubsystem->PopAllDirtyIADWithEntitiesPendingDestruction(DirtyIADs);
+
+	for (const TWeakObjectPtr<UInstancedActorsData>& DirtyIAD : DirtyIADs)
+	{
+		if (UInstancedActorsData* InstanceData = DirtyIAD.Get())
+		{
+			InstanceData->PopAllEntitiesPendingDestruction(EntitiesPendingDestruction);
+			if (EntitiesPendingDestruction.Num() > 0)
+			{
+				for (const FMassEntityHandle& EntityToDestroy : EntitiesPendingDestruction)
+				{
+					if (EntityManager.IsEntityValid(EntityToDestroy))
+					{
+						if (FMassRepresentationFragment* Representation = EntityManager.GetFragmentDataPtr<FMassRepresentationFragment>(EntityToDestroy))
+						{
+							Representation->bIsPendingDestruction = true;
+						}
+					}
+				}
+
+				Context.Defer().DestroyEntities(MoveTemp(EntitiesPendingDestruction));
+			}
+		}
+	}
+}
diff --git a/Engine/Plugins/Runtime/InstancedActors/Source/InstancedActors/Private/InstancedActorsSubsystem.cpp b/Engine/Plugins/Runtime/InstancedActors/Source/InstancedActors/Private/InstancedActorsSubsystem.cpp
--- a/Engine/Plugins/Runtime/InstancedActors/Source/InstancedActors/Private/InstancedActorsSubsystem.cpp
+++ b/Engine/Plugins/Runtime/InstancedActors/Source/InstancedActors/Private/InstancedActorsSubsystem.cpp
@@ -934,6 +934,18 @@ void UInstancedActorsSubsystem::UpdateAndResetTickTime(TConstStructView<FInstanc
 	}
 }
 
+void UInstancedActorsSubsystem::MarkIADDirtyWithEntitiesPendingDestruction(UInstancedActorsData& InstancedActorsData)
+{
+	FScopeLock Lock(&DirtyIadWithEntitiesPendingDestructionCS);
+	DirtyIadWithEntitiesPendingDestruction.AddUnique(&InstancedActorsData);
+}
+
+void UInstancedActorsSubsystem::PopAllDirtyIADWithEntitiesPendingDestruction(TArray<TWeakObjectPtr<UInstancedActorsData>>& OutInstancedActorsData)
+{
+	FScopeLock Lock(&DirtyIadWithEntitiesPendingDestructionCS);
+	OutInstancedActorsData.Append(MoveTemp(DirtyIadWithEntitiesPendingDestruction));
+}
+
 bool UInstancedActorsSubsystem::RegisterNewSharedFragmentsInternal(TConstStructView<FInstancedActorsDataSharedFragment> InstancedActorsDataSharedFragment)
 {
 	check(EntityManager);
diff --git a/Engine/Plugins/Runtime/InstancedActors/Source/InstancedActors/Public/InstancedActorsData.h b/Engine/Plugins/Runtime/InstancedActors/Source/InstancedActors/Public/InstancedActorsData.h
--- a/Engine/Plugins/Runtime/InstancedActors/Source/InstancedActors/Public/InstancedActorsData.h
+++ b/Engine/Plugins/Runtime/InstancedActors/Source/InstancedActors/Public/InstancedActorsData.h
@@ -6,6 +6,7 @@
 #include "InstancedActorsReplication.h"
 #include "MassEntityConfigAsset.h"
 #include "MassEntityTemplate.h"
+#include "HAL/CriticalSection.h"
 
 #include "InstancedActorsData.generated.h"
 
@@ -439,6 +440,13 @@ public:
 	UE_API int32 GetEntityIndexFromCollisionIndex(const UInstancedStaticMeshComponent& ISMComponent, const int32 CollisionIndex) const;
 
 	const FGameplayTagContainer& GetCombinedTags() const { return CombinedTags; }
+
+	/**
+	 * Moves the current list of entities pending destruction (added via QueueEntitiesForDestruction) to InOutEntities
+	 * This 'consumes' these pending entities with the assumption they will get destroyed by next command flush.
+	 * @see UInstancedActorsRuntimeRemovalProcessor, RuntimeRemoveInstances
+	 */
+	void PopAllEntitiesPendingDestruction(TArray<FMassEntityHandle>& InOutEntities);
 	
 protected:
 	FMassEntityHandle GetEntityHandleForIndex(const int32 Index) const
@@ -524,6 +532,14 @@ protected:
 	FBox CachedLocalBounds = FBox(ForceInit);
 
 private:
+	void QueueEntitiesForDestruction(TArray<FMassEntityHandle>&& EntitiesToDestroy);
+
+	mutable FCriticalSection PendingDestructionCS;
+
+	// These will get destroyed next time commands are flushed.
+	// @see QueueEntitiesForDestruction, PopAllEntitiesPendingDestruction
+	TArray<FMassEntityHandle> EntitiesPendingDestruction;
+
 	// Represents the shared fragment registered with MassEntityManager, that points back to this UInstancedActorsData instance
 	TStructView<FInstancedActorsDataSharedFragment> SharedInstancedActorDataStruct;
 
diff --git a/Engine/Plugins/Runtime/InstancedActors/Source/InstancedActors/Public/InstancedActorsRuntimeRemovalProcessor.h b/Engine/Plugins/Runtime/InstancedActors/Source/InstancedActors/Public/InstancedActorsRuntimeRemovalProcessor.h
new file mode 100644
--- /dev/null
+++ b/Engine/Plugins/Runtime/InstancedActors/Source/InstancedActors/Public/InstancedActorsRuntimeRemovalProcessor.h
@@ -0,0 +1,30 @@
+// Copyright Epic Games, Inc. All Rights Reserved.
+
+#pragma once
+
+#include "MassEntityQuery.h"
+#include "MassProcessor.h"
+
+#include "InstancedActorsRuntimeRemovalProcessor.generated.h"
+
+/** 
+ * This processor's sole responsibility is to find all UInstancedActorsData containing entities pending destruction
+ * and to actually schedule their destruction. It also marks their FMassRepresentationFragment as bIsPendingDestruction
+ * in order to prevent downstream processors from running any unnecessary representation logic on them.
+ * @see UInstancedActorsData::RuntimeRemoveInstances, UInstancedActorsSubsystem::MarkIadDirtyWithEntitiesPendingDestruction
+ */
+UCLASS(MinimalAPI)
+class UInstancedActorsRuntimeRemovalProcessor : public UMassProcessor
+{
+	GENERATED_BODY()
+public:
+	UInstancedActorsRuntimeRemovalProcessor();
+
+protected:
+	virtual bool ShouldAllowQueryBasedPruning(const bool bRuntimeMode = true) const override { return false; }
+	virtual void ConfigureQueries(const TSharedRef<FMassEntityManager>& EntityManager) override;
+	virtual void Execute(FMassEntityManager& EntityManager, FMassExecutionContext& Context) override;
+
+	/** Not directly used, but kept for the dependency solver */
+	FMassEntityQuery EntityQuery;
+};
\ No newline at end of file
diff --git a/Engine/Plugins/Runtime/InstancedActors/Source/InstancedActors/Public/InstancedActorsSubsystem.h b/Engine/Plugins/Runtime/InstancedActors/Source/InstancedActors/Public/InstancedActorsSubsystem.h
--- a/Engine/Plugins/Runtime/InstancedActors/Source/InstancedActors/Public/InstancedActorsSubsystem.h
+++ b/Engine/Plugins/Runtime/InstancedActors/Source/InstancedActors/Public/InstancedActorsSubsystem.h
@@ -209,6 +209,23 @@ public:
 		return InstancedActorsManagerClass; 
 	}
 
+	/**
+	 * Adds InstancedActorsData to a list of pending IAD which contains entities that should be destroyed
+	 * @see UInstancedActorsRuntimeRemovalProcessor, UInstancedActorsData::RuntimeRemoveInstances
+	 */
+	void MarkIADDirtyWithEntitiesPendingDestruction(UInstancedActorsData& InstancedActorsData);
+	/**
+	 * Moves the current list of IAD with entities pending destruction (added via MarkIadDirtyWithEntitiesPendingDestruction) to
+	 * OutInstancedActorsData and clears the internal list. This 'consumes' these pending IAD with the assumption the entities pending
+	 * destruction it contains will get destroyed by next command flush.
+	 * @see UInstancedActorsRuntimeRemovalProcessor, UInstancedActorsData::RuntimeRemoveInstances
+	 */
+	INSTANCEDACTORS_API void PopAllDirtyIADWithEntitiesPendingDestruction(TArray<TWeakObjectPtr<UInstancedActorsData>>& OutInstancedActorsData);
+	bool ContainsDirtyIADWithEntitiesPendingDestruction() const 
+	{ 
+		return DirtyIadWithEntitiesPendingDestruction.Num() > 0; 
+	}
+
 protected:
 	/** 
 	 * Fetches all registered FInstancedActorsDataSharedFragment from the EntityManager and adds the missing ones to SortedSharedFragments
@@ -254,6 +271,9 @@ protected:
 	// processing even out of 'detailed' representation processing range.
 	TArray<FInstancedActorsInstanceHandle> DirtyRepresentationInstances;
 
+	mutable FCriticalSection DirtyIadWithEntitiesPendingDestructionCS;
+	TArray<TWeakObjectPtr<UInstancedActorsData>> DirtyIadWithEntitiesPendingDestruction;
+
 #if WITH_INSTANCEDACTORS_DEBUG
 	TMap<TObjectKey<AInstancedActorsManager>, FBox> DebugManagerBounds;
 	TMap<TObjectKey<UInstancedActorsModifierVolumeComponent>, FBox> DebugModifierVolumeBounds;
diff --git a/Engine/Plugins/Runtime/MassGameplay/Source/MassRepresentation/Private/MassRepresentationProcessor.cpp b/Engine/Plugins/Runtime/MassGameplay/Source/MassRepresentation/Private/MassRepresentationProcessor.cpp
--- a/Engine/Plugins/Runtime/MassGameplay/Source/MassRepresentation/Private/MassRepresentationProcessor.cpp
+++ b/Engine/Plugins/Runtime/MassGameplay/Source/MassRepresentation/Private/MassRepresentationProcessor.cpp
@@ -87,6 +87,11 @@ void UMassRepresentationProcessor::UpdateRepresentation(FMassExecutionContext& C
 		FMassActorFragment& ActorInfo = ActorList[EntityIt];
 		AActor* Actor = ActorInfo.GetMutable();
 
+		if (Representation.bIsPendingDestruction)
+		{
+			continue;
+		}
+
 		// Keeping a copy of the that last calculated previous representation
 		const EMassRepresentationType PrevRepresentationCopy = Representation.PrevRepresentation;
 		Representation.PrevRepresentation = Representation.CurrentRepresentation;
diff --git a/Engine/Plugins/Runtime/MassGameplay/Source/MassRepresentation/Private/MassStationaryISMSwitcherProcessor.cpp b/Engine/Plugins/Runtime/MassGameplay/Source/MassRepresentation/Private/MassStationaryISMSwitcherProcessor.cpp
--- a/Engine/Plugins/Runtime/MassGameplay/Source/MassRepresentation/Private/MassStationaryISMSwitcherProcessor.cpp
+++ b/Engine/Plugins/Runtime/MassGameplay/Source/MassRepresentation/Private/MassStationaryISMSwitcherProcessor.cpp
@@ -58,8 +58,9 @@ void UMassStationaryISMSwitcherProcessor::ProcessContext(FMassExecutionContext&
 		const FTransformFragment& TransformFragment = TransformList[EntityIt];
 		FMassRepresentationFragment& Representation = RepresentationList[EntityIt];
 
-		if (Representation.PrevRepresentation != EMassRepresentationType::StaticMeshInstance
+		if ((Representation.PrevRepresentation != EMassRepresentationType::StaticMeshInstance
 			&& Representation.CurrentRepresentation != EMassRepresentationType::StaticMeshInstance)
+			|| Representation.bIsPendingDestruction)
 		{
 			// nothing to do here
 			continue;
diff --git a/Engine/Plugins/Runtime/MassGameplay/Source/MassRepresentation/Public/MassRepresentationFragments.h b/Engine/Plugins/Runtime/MassGameplay/Source/MassRepresentation/Public/MassRepresentationFragments.h
--- a/Engine/Plugins/Runtime/MassGameplay/Source/MassRepresentation/Public/MassRepresentationFragments.h
+++ b/Engine/Plugins/Runtime/MassGameplay/Source/MassRepresentation/Public/MassRepresentationFragments.h
@@ -50,32 +50,35 @@ struct FMassRepresentationLODFragment : public FMassFragment
 USTRUCT()
 struct FMassRepresentationFragment : public FMassFragment
 {
-	GENERATED_BODY()
+    GENERATED_BODY()
 
-	UPROPERTY()
-	EMassRepresentationType CurrentRepresentation = EMassRepresentationType::None;
+    UPROPERTY()
+    EMassRepresentationType CurrentRepresentation = EMassRepresentationType::None;
 
-	UPROPERTY()
-	EMassRepresentationType PrevRepresentation = EMassRepresentationType::None;
+    UPROPERTY()
+    EMassRepresentationType PrevRepresentation = EMassRepresentationType::None;
 
-	UPROPERTY()
-	int16 HighResTemplateActorIndex = INDEX_NONE;
+    UPROPERTY()
+    int16 HighResTemplateActorIndex = INDEX_NONE;
 
-	UPROPERTY()
-	int16 LowResTemplateActorIndex = INDEX_NONE;
+    UPROPERTY()
+    int16 LowResTemplateActorIndex = INDEX_NONE;
 
-	UPROPERTY()
-	FStaticMeshInstanceVisualizationDescHandle StaticMeshDescHandle;
+    UPROPERTY()
+    FStaticMeshInstanceVisualizationDescHandle StaticMeshDescHandle;
 
-	UPROPERTY()
-	FMassActorSpawnRequestHandle ActorSpawnRequestHandle;
+    /** Value scaling from 0 to 3, 0 highest LOD we support and 3 being completely off LOD */
+    UPROPERTY()
+    float PrevLODSignificance = -1.0f;
 
-	UPROPERTY()
-	FTransform PrevTransform;
+    UPROPERTY()
+    uint8 bIsPendingDestruction : 1  = false;
 
-	/** Value scaling from 0 to 3, 0 highest LOD we support and 3 being completely off LOD */
-	UPROPERTY()
-	float PrevLODSignificance = -1.0f;
+    UPROPERTY()
+    FMassActorSpawnRequestHandle ActorSpawnRequestHandle;
+
+    UPROPERTY()
+    FTransform PrevTransform;
 };
 
 USTRUCT()
