Added the runtime grid to IAM's audit instances debug console command

#Synced-CL 48612369

diff --git a/Engine/Plugins/Runtime/InstancedActors/Source/InstancedActors/Private/InstancedActorsManager.cpp b/Engine/Plugins/Runtime/InstancedActors/Source/InstancedActors/Private/InstancedActorsManager.cpp
--- a/Engine/Plugins/Runtime/InstancedActors/Source/InstancedActors/Private/InstancedActorsManager.cpp
+++ b/Engine/Plugins/Runtime/InstancedActors/Source/InstancedActors/Private/InstancedActorsManager.cpp
@@ -1221,6 +1221,9 @@ void AInstancedActorsManager::AuditInstances(FOutputDevice& Ar, bool bDebugDraw,
 	Ar.Logf(TEXT("\tLevel Transform: %s"), *GetLevelTransform().ToString());
 	Ar.Logf(TEXT("\tLocation: %s"), *GetActorLocation().ToString());
 	Ar.Logf(TEXT("\tNum Instances: %d"), GetNumValidInstances());
+#if WITH_EDITOR
+	Ar.Logf(TEXT("\tRuntime grid: %s"), *GetRuntimeGrid().ToString());
+#endif
 
 #if UE_ENABLE_DEBUG_DRAWING
 	if (bDebugDraw)
