Removed ZeroConstructor trait for FInstancedActorsInstanceIndex as it conflicts with the new default INSTANCED_INDEX_NONE sentinel value

#Synced-CL 48806929

diff --git a/Engine/Plugins/Runtime/InstancedActors/Source/InstancedActors/Public/InstancedActorsIndex.h b/Engine/Plugins/Runtime/InstancedActors/Source/InstancedActors/Public/InstancedActorsIndex.h
--- a/Engine/Plugins/Runtime/InstancedActors/Source/InstancedActors/Public/InstancedActorsIndex.h
+++ b/Engine/Plugins/Runtime/InstancedActors/Source/InstancedActors/Public/InstancedActorsIndex.h
@@ -25,7 +25,7 @@ struct FInstancedActorsInstanceIndex
 	FInstancedActorsInstanceIndex() = default;
 	explicit FInstancedActorsInstanceIndex(const int32 InIndex)
 	{ 
-		check(InIndex >= 0 && InIndex < INSTANCED_INDEX_NONE);
+		check(InIndex >= 0 && InIndex <= INSTANCED_INDEX_MAX);
 		Index = InIndex;
 	}
 
@@ -88,7 +88,6 @@ struct TStructOpsTypeTraits<FInstancedActorsInstanceIndex> : public TStructOpsTy
 {
 	enum
 	{
-		WithZeroConstructor = true,
 		WithCopy = true,
 		WithIdenticalViaEquality = true,
 	};
