InstancedActorsData's ApplyInstanceDeltas now uses the InstanceDeltaIndices parameter instead of iterating on all deltas

#Synced-CL 49026660

diff --git a/Engine/Plugins/Runtime/InstancedActors/Source/InstancedActors/Private/InstancedActorsData.cpp b/Engine/Plugins/Runtime/InstancedActors/Source/InstancedActors/Private/InstancedActorsData.cpp
--- a/Engine/Plugins/Runtime/InstancedActors/Source/InstancedActors/Private/InstancedActorsData.cpp
+++ b/Engine/Plugins/Runtime/InstancedActors/Source/InstancedActors/Private/InstancedActorsData.cpp
@@ -1143,12 +1143,18 @@ void UInstancedActorsData::ApplyInstanceDeltas(TConstArrayView<int32> InstanceDe
 
 	const TArray<FInstancedActorsDelta>& Deltas = InstanceDeltas.GetInstanceDeltas();
 
-	for (const FInstancedActorsDelta& Delta : Deltas)
-	{
-		ApplyInstanceDelta(EntityManager, Delta, EntitiesToRemove);
-	}
+	for (int32 InstanceDeltaIndex : InstanceDeltaIndices)
+    {
+        if (!ensureMsgf(Deltas.IsValidIndex(InstanceDeltaIndex), TEXT("Invalid InstanceDelta index %d (Num=%d) in ApplyInstanceDeltas"), InstanceDeltaIndex, Deltas.Num()))
+        {
+            continue;
+        }
 
-	RuntimeRemoveInstances(MakeArrayView(EntitiesToRemove));
+        const FInstancedActorsDelta& Delta = Deltas[InstanceDeltaIndex];
+        ApplyInstanceDelta(EntityManager, Delta, EntitiesToRemove);
+    }
+
+    RuntimeRemoveInstances(MakeArrayView(EntitiesToRemove));
 }
 
 void UInstancedActorsData::RollbackInstanceDeltas(TConstArrayView<int32> InstanceDeltaIndices)
