PerActorClassInstanceData and the IAD's ActorClass are now exposed as BlueprintReadOnly

#Synced-CL 49026660

diff --git a/Engine/Plugins/Runtime/InstancedActors/Source/InstancedActors/Public/InstancedActorsData.h b/Engine/Plugins/Runtime/InstancedActors/Source/InstancedActors/Public/InstancedActorsData.h
--- a/Engine/Plugins/Runtime/InstancedActors/Source/InstancedActors/Public/InstancedActorsData.h
+++ b/Engine/Plugins/Runtime/InstancedActors/Source/InstancedActors/Public/InstancedActorsData.h
@@ -104,7 +104,7 @@ public:
 	uint16 ID;
 
 	// The fully realized actor type for these instances
-	UPROPERTY(VisibleAnywhere, Category=InstancedActors)
+	UPROPERTY(VisibleAnywhere, BlueprintReadOnly, Category=InstancedActors)
 	TSubclassOf<AActor> ActorClass;
 
 	const FInstancedActorsTagSet& GetAdditionalTags() const
diff --git a/Engine/Plugins/Runtime/InstancedActors/Source/InstancedActors/Public/InstancedActorsManager.h b/Engine/Plugins/Runtime/InstancedActors/Source/InstancedActors/Public/InstancedActorsManager.h
--- a/Engine/Plugins/Runtime/InstancedActors/Source/InstancedActors/Public/InstancedActorsManager.h
+++ b/Engine/Plugins/Runtime/InstancedActors/Source/InstancedActors/Public/InstancedActorsManager.h
@@ -349,7 +349,7 @@ protected:
 	uint16 NextInstanceDataID = 0;
 
 	/** Per-actor-class instance data populated by AddActorInstance */
-	UPROPERTY(Instanced, VisibleAnywhere, Category=InstancedActors)
+	UPROPERTY(Instanced, VisibleAnywhere, BlueprintReadOnly, Category=InstancedActors)
 	TArray<TObjectPtr<UInstancedActorsData>> PerActorClassInstanceData;
 
 	/** World space cumulative instance bounds, calculated in BeginPlay */
