Expose InstancedActorsData's InstanceTransforms to blueprints

#Synced-CL 50358563
#Synced-CL 50358563

diff --git a/Engine/Plugins/Runtime/InstancedActors/Source/InstancedActors/Public/InstancedActorsData.h b/Engine/Plugins/Runtime/InstancedActors/Source/InstancedActors/Public/InstancedActorsData.h
--- a/Engine/Plugins/Runtime/InstancedActors/Source/InstancedActors/Public/InstancedActorsData.h
+++ b/Engine/Plugins/Runtime/InstancedActors/Source/InstancedActors/Public/InstancedActorsData.h
@@ -180,7 +180,7 @@ public:
 	// Per-instance transforms. This essentially forms the 'instance list' until Entities are spawned from this.
 	// Note: This list can contain invalid 'removed' entries with 0'd values which will be skipped by
 	//       UInstancedActorsInitializerProcessor
-	UPROPERTY()
+	UPROPERTY(VisibleAnywhere, BlueprintReadOnly, Category=InstancedActors)
 	TArray<FTransform> InstanceTransforms;
 
 	// InstanceTransforms.Num() - number of invalid transforms
