New InstancedActorsEditorSubsystem, exposing some InstancedActors functionnality.
* Added the ability to retrieve an instance's handle
* Exposed the InstancedActorsData's SetInstanceTransform

#Synced-CL 50980885

diff --git a/Engine/Plugins/Runtime/InstancedActors/Source/InstancedActors/Public/InstancedActorsData.h b/Engine/Plugins/Runtime/InstancedActors/Source/InstancedActors/Public/InstancedActorsData.h
--- a/Engine/Plugins/Runtime/InstancedActors/Source/InstancedActors/Public/InstancedActorsData.h
+++ b/Engine/Plugins/Runtime/InstancedActors/Source/InstancedActors/Public/InstancedActorsData.h
@@ -242,6 +242,7 @@ public:
 
 	// Set/update the transform for an instance.
 	//
+	UFUNCTION(BlueprintCallable, Category = InstancedActors)
 	UE_API bool SetInstanceTransform(const FInstancedActorsInstanceHandle& InstanceHandle, const FTransform& Transform, const bool bWorldSpace);
 
 	// Removes all instance data for InstanceHandle.
diff --git a/Engine/Plugins/Runtime/InstancedActors/Source/InstancedActorsEditor/InstancedActorsEditor.Build.cs b/Engine/Plugins/Runtime/InstancedActors/Source/InstancedActorsEditor/InstancedActorsEditor.Build.cs
--- a/Engine/Plugins/Runtime/InstancedActors/Source/InstancedActorsEditor/InstancedActorsEditor.Build.cs
+++ b/Engine/Plugins/Runtime/InstancedActors/Source/InstancedActorsEditor/InstancedActorsEditor.Build.cs
@@ -24,7 +24,8 @@ namespace UnrealBuildTool.Rules
 					"Slate",
 					"LevelEditor",
 					"InstancedActors",
-					"MassEntity"
+					"MassEntity",
+					"EditorSubsystem"
 				}
 			);
 		}
diff --git a/Engine/Plugins/Runtime/InstancedActors/Source/InstancedActorsEditor/Private/InstancedActorsEditorSubsystem.cpp b/Engine/Plugins/Runtime/InstancedActors/Source/InstancedActorsEditor/Private/InstancedActorsEditorSubsystem.cpp
new file mode 100644
--- /dev/null
+++ b/Engine/Plugins/Runtime/InstancedActors/Source/InstancedActorsEditor/Private/InstancedActorsEditorSubsystem.cpp
@@ -0,0 +1,17 @@
+// Copyright Epic Games, Inc. All Rights Reserved.
+
+#include "InstancedActorsEditorSubsystem.h"
+#include "InstancedActorsIndex.h"
+#include "InstancedActorsData.h"
+
+FInstancedActorsInstanceHandle UInstancedActorsEditorSubsystem::GetInstanceHandle(UInstancedActorsData* InstanceData, int32 Index)
+{
+	if (!InstanceData || Index < 0 || Index >= InstanceData->InstanceTransforms.Num())
+	{
+		UE_LOGF(LogInstancedActors, Error, "Invalid data, returning default handle.");
+		return FInstancedActorsInstanceHandle();
+	}
+
+	const FInstancedActorsInstanceHandle Handle(*InstanceData, FInstancedActorsInstanceIndex(Index));
+	return Handle;
+}
diff --git a/Engine/Plugins/Runtime/InstancedActors/Source/InstancedActorsEditor/Public/InstancedActorsEditorSubsystem.h b/Engine/Plugins/Runtime/InstancedActors/Source/InstancedActorsEditor/Public/InstancedActorsEditorSubsystem.h
new file mode 100644
--- /dev/null
+++ b/Engine/Plugins/Runtime/InstancedActors/Source/InstancedActorsEditor/Public/InstancedActorsEditorSubsystem.h
@@ -0,0 +1,29 @@
+// Copyright Epic Games, Inc. All Rights Reserved.
+
+#pragma once
+
+#include "CoreMinimal.h"
+#include "EditorSubsystem.h"
+#include "InstancedActorsIndex.h"
+
+#include "InstancedActorsEditorSubsystem.generated.h"
+
+#define UE_API INSTANCEDACTORSEDITOR_API
+
+class UInstancedActorsData;
+
+/**
+* Subsystem for exposing instanced actors functionality to blueprints
+*/
+UCLASS(MinimalAPI)
+class UInstancedActorsEditorSubsystem : public UEditorSubsystem
+{
+	GENERATED_BODY()
+
+public:
+
+	UFUNCTION(BlueprintPure, Category = "InstancedActors", meta = (DisplayName = "Get Instance Handle (by Index)"))
+	static UE_API FInstancedActorsInstanceHandle GetInstanceHandle(UInstancedActorsData* InstanceData, int32 Index);
+};
+
+#undef UE_API
