InstancedActorsData now defers entity destruction in RuntimeRemoveAllInstances

#Synced-CL 51173315

diff --git a/Engine/Plugins/Runtime/InstancedActors/Source/InstancedActors/Private/InstancedActorsData.cpp b/Engine/Plugins/Runtime/InstancedActors/Source/InstancedActors/Private/InstancedActorsData.cpp
--- a/Engine/Plugins/Runtime/InstancedActors/Source/InstancedActors/Private/InstancedActorsData.cpp
+++ b/Engine/Plugins/Runtime/InstancedActors/Source/InstancedActors/Private/InstancedActorsData.cpp
@@ -1301,15 +1301,26 @@ void UInstancedActorsData::RuntimeRemoveAllInstances()
 	// Destroy spawned entities if we've spawned them already
 	if (HasSpawnedEntities())
 	{
-		TArray<FMassArchetypeEntityCollection> EntityCollectionsToDestroy;
-		UE::Mass::Utils::CreateEntityCollections(MassEntityManager, Entities, FMassArchetypeEntityCollection::NoDuplicates, EntityCollectionsToDestroy);
+		if (UE::InstancedActors::CVars::bDeferRuntimeRemoveInstancesEntityDestruction)
+		{
+			QueueEntitiesForDestruction(MoveTemp(Entities));
+			if (UInstancedActorsSubsystem* IASubsystem = GetManagerChecked().GetInstancedActorSubsystem())
+			{
+				IASubsystem->MarkIADDirtyWithEntitiesPendingDestruction(*this);
+			}
+		}
+		else
+		{
+			TArray<FMassArchetypeEntityCollection> EntityCollectionsToDestroy;
+			UE::Mass::Utils::CreateEntityCollections(MassEntityManager, Entities, FMassArchetypeEntityCollection::NoDuplicates, EntityCollectionsToDestroy);
 
-		// Destroy entities. This will also trigger actor destruction for any spawned actors
-		// for these entities.
-		MassEntityManager.BatchDestroyEntityChunks(EntityCollectionsToDestroy);
+			// Destroy entities. This will also trigger actor destruction for any spawned actors
+			// for these entities.
+			MassEntityManager.BatchDestroyEntityChunks(EntityCollectionsToDestroy);
 
-		// Zero out all entity handles to 'reset' them
-		FMemory::Memzero(Entities.GetData(), Entities.GetTypeSize() * Entities.Num());
+			// Zero out all entity handles to 'reset' them
+			FMemory::Memzero(Entities.GetData(), Entities.GetTypeSize() * Entities.Num());
+		}
 	}
 	// Pre-empt entity spawning and simply invalidate InstanceTransform entries, preventing them from spawning later
 	else
@@ -2012,4 +2023,4 @@ int32 UInstancedActorsData::GetEntityIndexFromCollisionIndex(const UInstancedSta
 	VLOG_COLLISIONINDEX(this, LogInstancedActors, Verbose, TEXT("Found EntityIndex  %d"), EntityIndex);
 
 	return EntityIndex;
-}
+}
\ No newline at end of file
