Added a protected virtual OnEditorPreviewISMCsCreated() to UInstancedActorsData called after editor preview ISMCs are first created. This allows subclasses to customize the preview components without modifying the base creation flow.

#Synced-CL 51230964

diff --git a/Engine/Plugins/Runtime/InstancedActors/Source/InstancedActors/Private/InstancedActorsManager.cpp b/Engine/Plugins/Runtime/InstancedActors/Source/InstancedActors/Private/InstancedActorsManager.cpp
--- a/Engine/Plugins/Runtime/InstancedActors/Source/InstancedActors/Private/InstancedActorsManager.cpp
+++ b/Engine/Plugins/Runtime/InstancedActors/Source/InstancedActors/Private/InstancedActorsManager.cpp
@@ -820,9 +820,12 @@ UInstancedActorsData& AInstancedActorsManager::GetOrCreateActorInstanceData(TSub
 
 		// Add EditorPreviewISMComponents to the known ISMComponent-to-InstanceData map (ISMComponentToInstanceDataMap)
 		// so that we can identify Instanced Actors instances by ISM instance id coming from one of the EditorPreviewISMComponents
-		// (like when the instance is clicked in the editor). 
-		// @see ActorInstanceHandleFromFSMInstanceId 
-		RegisterInstanceDatasComponents(*NewInstanceData, NewInstanceData->EditorPreviewISMComponents);		
+		// (like when the instance is clicked in the editor).
+		// @see ActorInstanceHandleFromFSMInstanceId
+		RegisterInstanceDatasComponents(*NewInstanceData, NewInstanceData->EditorPreviewISMComponents);
+
+		// Allow subclasses to customize the freshly created preview ISMCs (e.g. apply material overrides).
+		NewInstanceData->OnEditorPreviewISMCsCreated();
 	}
 
 	// Cache asset bounds for use during instance population
diff --git a/Engine/Plugins/Runtime/InstancedActors/Source/InstancedActors/Public/InstancedActorsData.h b/Engine/Plugins/Runtime/InstancedActors/Source/InstancedActors/Public/InstancedActorsData.h
--- a/Engine/Plugins/Runtime/InstancedActors/Source/InstancedActors/Public/InstancedActorsData.h
+++ b/Engine/Plugins/Runtime/InstancedActors/Source/InstancedActors/Public/InstancedActorsData.h
@@ -152,6 +152,14 @@ protected:
 	{
 	}
 
+#if WITH_EDITOR
+	// Called from GetOrCreateActorInstanceData after EditorPreviewISMComponents are first created,
+	// allowing subclasses to customize them (e.g. apply material overrides).
+	virtual void OnEditorPreviewISMCsCreated()
+	{
+	}
+#endif
+
 	// Delimiting tag set to group sets of instances
 	UPROPERTY(VisibleAnywhere, Category=InstancedActors)
 	FInstancedActorsTagSet AdditionalTags;
