UInstancedActorsStationaryLODBatchProcessor now requires game thread execution

#Synced-CL 51230964

diff --git a/Engine/Plugins/Runtime/InstancedActors/Source/InstancedActors/Private/InstancedActorsStationaryLODBatchProcessor.cpp b/Engine/Plugins/Runtime/InstancedActors/Source/InstancedActors/Private/InstancedActorsStationaryLODBatchProcessor.cpp
--- a/Engine/Plugins/Runtime/InstancedActors/Source/InstancedActors/Private/InstancedActorsStationaryLODBatchProcessor.cpp
+++ b/Engine/Plugins/Runtime/InstancedActors/Source/InstancedActors/Private/InstancedActorsStationaryLODBatchProcessor.cpp
@@ -118,7 +118,8 @@ UInstancedActorsStationaryLODBatchProcessor::UInstancedActorsStationaryLODBatchP
 
 	ExecutionFlags = (int32)EProcessorExecutionFlags::AllNetModes;
 
-	bRequiresGameThreadExecution = false;
+     // Direct ISM component modifications are not thread-safe and must run on the game thread.
+	bRequiresGameThreadExecution = true;
 
 	DelayPerBulkLOD[(int)EInstancedActorsBulkLOD::Detailed] = 5.0;
 	DelayPerBulkLOD[(int)EInstancedActorsBulkLOD::Medium] = 1.0;
