GetOrCompileSettingsForActorClass is now exported

#Synced-CL 51506951

diff --git a/Engine/Plugins/Runtime/InstancedActors/Source/InstancedActors/Public/InstancedActorsSubsystem.h b/Engine/Plugins/Runtime/InstancedActors/Source/InstancedActors/Public/InstancedActorsSubsystem.h
--- a/Engine/Plugins/Runtime/InstancedActors/Source/InstancedActors/Public/InstancedActorsSubsystem.h
+++ b/Engine/Plugins/Runtime/InstancedActors/Source/InstancedActors/Public/InstancedActorsSubsystem.h
@@ -145,7 +145,7 @@ public:
 	 *       for quick lookup in CompileSettingsForActorClass. This means unique class names must be used for 
 	 *       per-class settings.
 	 */
-	FSharedStruct GetOrCompileSettingsForActorClass(TSubclassOf<AActor> ActorClass);
+	INSTANCEDACTORS_API FSharedStruct GetOrCompileSettingsForActorClass(TSubclassOf<AActor> ActorClass);
 
 	/**
 	 * Returns true if ActorClass has a matching FInstancedActorsClassSettingsBase entry in 
