[InstancedActors] Fix deferred lambda capturing actor by reference instead of weak pointer

SetActorEnabled captured AActor& by reference in a FMassDeferredSetCommand lambda. Since deferred commands execute later when the command buffer is flushed, the actor could be destroyed/GCd between capture and execution, causing a dangling reference. Capture TWeakObjectPtr<AActor> instead and check validity before use.

#rnx

#Synced-CL 51621734

diff --git a/Engine/Plugins/Runtime/InstancedActors/Source/InstancedActors/Private/InstancedActorsRepresentationActorManagement.cpp b/Engine/Plugins/Runtime/InstancedActors/Source/InstancedActors/Private/InstancedActorsRepresentationActorManagement.cpp
--- a/Engine/Plugins/Runtime/InstancedActors/Source/InstancedActors/Private/InstancedActorsRepresentationActorManagement.cpp
+++ b/Engine/Plugins/Runtime/InstancedActors/Source/InstancedActors/Private/InstancedActorsRepresentationActorManagement.cpp
@@ -66,9 +66,12 @@ void UInstancedActorsRepresentationActorManagement::SetActorEnabled(const EMassA
 		if (Actor.GetActorEnableCollision())
 		{
 			// Deferring this as there is a callback internally that could end up doing things outside of the game thread and will fire checks(Chaos mostly)
-			CommandBuffer.PushCommand<FMassDeferredSetCommand>([&Actor](FMassEntityManager&)
+			CommandBuffer.PushCommand<FMassDeferredSetCommand>([WeakActor = TWeakObjectPtr<AActor>(&Actor)](FMassEntityManager&)
 				{
-					Actor.SetActorEnableCollision(true);
+					if (AActor* ActorPtr = WeakActor.Get())
+					{
+						ActorPtr->SetActorEnableCollision(true);
+					}
 				});
 		}
 
