Shared fragment nullification moved from DespawnEntities into Deinitialize, which is only called on true teardown (EndPlay). This preserves the fragment for immediate reuse by the incoming spawn when switching subsystems and removes the anti-pattern of a boolean parameter.

#Synced-CL 51897778

diff --git a/Engine/Plugins/Runtime/InstancedActors/Source/InstancedActors/Private/InstancedActorsData.cpp b/Engine/Plugins/Runtime/InstancedActors/Source/InstancedActors/Private/InstancedActorsData.cpp
--- a/Engine/Plugins/Runtime/InstancedActors/Source/InstancedActors/Private/InstancedActorsData.cpp
+++ b/Engine/Plugins/Runtime/InstancedActors/Source/InstancedActors/Private/InstancedActorsData.cpp
@@ -305,6 +305,28 @@ void UInstancedActorsData::Initialize()
 
 void UInstancedActorsData::Deinitialize()
 {
+	// Nullify the shared fragment now that this IAD is truly going away. This is intentionally not done in
+	// DespawnEntities() so that on subsystem switch (where Deinitialize is not called), the fragment remains
+	// intact and can be reused by the immediately following respawn.
+	FMassEntityManager& MassEntityManager = GetMassEntityManagerChecked();
+	FInstancedActorsDataSharedFragment ManagerFragment;
+	ManagerFragment.InstanceData = this;
+	FSharedStruct SharedFragmentInstance = MassEntityManager.GetOrCreateSharedFragment(ManagerFragment);
+	if (FInstancedActorsDataSharedFragment* AsSharedManagerFragment = SharedFragmentInstance.GetPtr<FInstancedActorsDataSharedFragment>())
+	{
+		if (ensure(AsSharedManagerFragment->InstanceData == this))
+		{
+			if (UE::InstancedActors::CVars::bFlagSharedFragmentInsteadOfNullify)
+			{
+				AsSharedManagerFragment->bInstancedActorsDataDespawnedEntities = true;
+			}
+			else
+			{
+				AsSharedManagerFragment->InstanceData = nullptr;
+			}
+		}
+	}
+
 	if (bHasEverInitialized && UE::InstancedActors::CVars::bEnableReleasingEntityTemplatesAndExemplarActors)
 	{
 		ReleaseEntityTemplate();
@@ -455,7 +477,7 @@ void UInstancedActorsData::SpawnEntities()
 	OnSpawnEntities();
 }
 
-void UInstancedActorsData::DespawnEntities(bool bIsSwitchingSubsystems /*= false*/)
+void UInstancedActorsData::DespawnEntities()
 {
 	UWorld* World = GetWorld();
 	check(World);
@@ -550,20 +572,8 @@ void UInstancedActorsData::DespawnEntities(bool bIsSwitchingSubsystems /*= false
 	FInstancedActorsDataSharedFragment* AsSharedManagerFragment = SharedFragmentInstance.GetPtr<FInstancedActorsDataSharedFragment>();
 	if (ensure(AsSharedManagerFragment) && AsSharedManagerFragment->InstanceData == this)
 	{
-		// decrement the stats before resetting
+		// Decrement the stats before resetting
 		Manager.UpdateInstanceStats(NumInstances, AsSharedManagerFragment->BulkLOD, /*Increment=*/false);
-
-		if (!bIsSwitchingSubsystems)
-		{
-			if (UE::InstancedActors::CVars::bFlagSharedFragmentInsteadOfNullify)
-			{
-				AsSharedManagerFragment->bInstancedActorsDataDespawnedEntities = true;
-			}
-			else
-			{
-				AsSharedManagerFragment->InstanceData = nullptr;
-			}
-		}
 		AsSharedManagerFragment->BulkLOD = EInstancedActorsBulkLOD::MAX;
 	}
 
@@ -579,6 +589,11 @@ void UInstancedActorsData::DespawnEntities(bool bIsSwitchingSubsystems /*= false
 	OnDespawnEntities();
 }
 
+void UInstancedActorsData::DespawnEntities(bool /*bIsSwitchingSubsystems*/)
+{
+	DespawnEntities();
+}
+
 AInstancedActorsManager* UInstancedActorsData::GetManager() const
 {
 	if (UObject* Outer = GetOuter())
diff --git a/Engine/Plugins/Runtime/InstancedActors/Source/InstancedActors/Private/InstancedActorsManager.cpp b/Engine/Plugins/Runtime/InstancedActors/Source/InstancedActors/Private/InstancedActorsManager.cpp
--- a/Engine/Plugins/Runtime/InstancedActors/Source/InstancedActors/Private/InstancedActorsManager.cpp
+++ b/Engine/Plugins/Runtime/InstancedActors/Source/InstancedActors/Private/InstancedActorsManager.cpp
@@ -286,7 +286,7 @@ void AInstancedActorsManager::OnAddedToSubsystem(UInstancedActorsSubsystem& InIn
 	if (IsValid(InstancedActorSubsystem) && (InstancedActorSubsystem != &InInstancedActorSubsystem))
 	{
 		// we need to unregister from the previous subsystem first
-		DespawnAllEntities(/*bIsSwitchingSubsystems*/ true);
+		DespawnAllEntities();
 		// and let registered modifiers go
 		RemoveAllModifierVolumes();
 		
@@ -914,20 +914,24 @@ void AInstancedActorsManager::RuntimeRemoveAllInstances()
 	}
 }
 
-void AInstancedActorsManager::DespawnAllEntities(bool bIsSwitchingSubsystems /*= false*/)
+void AInstancedActorsManager::DespawnAllEntities()
 {
-	// DespawnEntities for all PerActorClassInstanceData
 	if (HasSpawnedEntities())
 	{
+		// DespawnEntities for all PerActorClassInstanceData
 		for (TObjectPtr<UInstancedActorsData> InstanceData : PerActorClassInstanceData)
 		{
-			// Reconstruct instance data from Mass then destroy all Mass entities
-			InstanceData->DespawnEntities(bIsSwitchingSubsystems);
+			InstanceData->DespawnEntities();
 		}
 		bHasSpawnedEntities = false;
 	}
 }
 
+void AInstancedActorsManager::DespawnAllEntities(bool /*bIsSwitchingSubsystems*/)
+{
+	DespawnAllEntities();
+}
+
 bool AInstancedActorsManager::ForEachInstance(FInstanceOperationFunc Operation) const
 {
 	FScopedInstancedActorsIterationContext IterationContext;
diff --git a/Engine/Plugins/Runtime/InstancedActors/Source/InstancedActors/Public/InstancedActorsData.h b/Engine/Plugins/Runtime/InstancedActors/Source/InstancedActors/Public/InstancedActorsData.h
--- a/Engine/Plugins/Runtime/InstancedActors/Source/InstancedActors/Public/InstancedActorsData.h
+++ b/Engine/Plugins/Runtime/InstancedActors/Source/InstancedActors/Public/InstancedActorsData.h
@@ -64,12 +64,17 @@ public:
 
 	// Called in AInstancedActorsManager::InitializeModifyAndSpawnEntities to spawn Mass entities for each instance
 	UE_API void SpawnEntities();
-
 	// Called early in AInstancedActorsManager::EndPlay to reconstruct cooked data state from runtime Mass entities as best we can,
-	// then despawn all Mass entities and reset any other runtime instance data
-	UE_API void DespawnEntities(bool bIsSwitchingSubsystems = false);
+	// then despawn all Mass entities and reset any other runtime instance data.
+	// Note: does NOT nullify the shared fragment, that is done in Deinitialize(), which is only called on true teardown
+	// and not on subsystem switch, allowing the fragment to be reused by an immediate respawn.
+	UE_API void DespawnEntities();
+
+	UE_DEPRECATED(5.8, "Using DespawnEntities with a parameter is deprecated.")
+	UE_API void DespawnEntities(bool bIsSwitchingSubsystems);
 
-	// Called from the parent's AInstancedActorsManager::EndPlay. Can release the entity template and exemplar Actor from memory 
+	// Called from the parent's AInstancedActorsManager::EndPlay. Nullifies the shared fragment and can release
+	// the entity template and exemplar Actor from memory.
 	UE_API void Deinitialize();
 
 	UFUNCTION(BlueprintPure, Category = InstancedActors)
diff --git a/Engine/Plugins/Runtime/InstancedActors/Source/InstancedActors/Public/InstancedActorsManager.h b/Engine/Plugins/Runtime/InstancedActors/Source/InstancedActors/Public/InstancedActorsManager.h
--- a/Engine/Plugins/Runtime/InstancedActors/Source/InstancedActors/Public/InstancedActorsManager.h
+++ b/Engine/Plugins/Runtime/InstancedActors/Source/InstancedActors/Public/InstancedActorsManager.h
@@ -292,7 +292,10 @@ protected:
 	UE_API void SerializeInstancePersistenceData(FStructuredArchive::FRecord Record, UInstancedActorsData* InstanceData, int64 TimeDelta) const;
 
 	/** Despawns all entities spawned by individual UInstancedActorsData instances. */
-	UE_API virtual void DespawnAllEntities(bool bIsSwitchingSubsystems = false);
+	UE_API virtual void DespawnAllEntities();
+
+	UE_DEPRECATED(5.8, "Using DespawnAllEntities with a parameter is deprecated.")
+	UE_API void DespawnAllEntities(bool bIsSwitchingSubsystems);
 
 	/** Attempts to run any 'pending' modifiers in ModifierVolumes where are appropriate to run given HasSpawnedEntities
 	 * Called in BeginPlay prior to, and then again after SpawnEntities. Also called in AddModifierVolume.
